Fixing Gameplugins build

diff --git a/Engine/Plugins/Online/HTTPChunkInstaller/Config/Localization/HTTPChunkInstaller.ini b/Engine/Plugins/Online/HTTPChunkInstaller/Config/Localization/HTTPChunkInstaller.ini
deleted file mode 100644
diff --git a/Engine/Plugins/Online/HTTPChunkInstaller/HTTPChunkInstaller.uplugin b/Engine/Plugins/Online/HTTPChunkInstaller/HTTPChunkInstaller.uplugin
deleted file mode 100644
diff --git a/Engine/Plugins/Online/HTTPChunkInstaller/Source/HTTPChunkInstaller.Build.cs b/Engine/Plugins/Online/HTTPChunkInstaller/Source/HTTPChunkInstaller.Build.cs
deleted file mode 100644
diff --git a/Engine/Plugins/Online/HTTPChunkInstaller/Source/Private/HTTPChunkInstaller.cpp b/Engine/Plugins/Online/HTTPChunkInstaller/Source/Private/HTTPChunkInstaller.cpp
deleted file mode 100644
diff --git a/Engine/Plugins/Online/HTTPChunkInstaller/Source/Private/HTTPChunkInstallerLog.cpp b/Engine/Plugins/Online/HTTPChunkInstaller/Source/Private/HTTPChunkInstallerLog.cpp
deleted file mode 100644
diff --git a/Engine/Plugins/Online/HTTPChunkInstaller/Source/Private/HTTPChunkInstallerLog.h b/Engine/Plugins/Online/HTTPChunkInstaller/Source/Private/HTTPChunkInstallerLog.h
deleted file mode 100644
diff --git a/Engine/Plugins/Online/HTTPChunkInstaller/Source/Private/HTTPOnlineTitleFile.cpp b/Engine/Plugins/Online/HTTPChunkInstaller/Source/Private/HTTPOnlineTitleFile.cpp
deleted file mode 100644
diff --git a/Engine/Plugins/Online/HTTPChunkInstaller/Source/Private/LocalTitleFile.cpp b/Engine/Plugins/Online/HTTPChunkInstaller/Source/Private/LocalTitleFile.cpp
deleted file mode 100644
diff --git a/Engine/Plugins/Online/HTTPChunkInstaller/Source/Public/ChunkInstall.h b/Engine/Plugins/Online/HTTPChunkInstaller/Source/Public/ChunkInstall.h
deleted file mode 100644
diff --git a/Engine/Plugins/Online/HTTPChunkInstaller/Source/Public/ChunkSetup.h b/Engine/Plugins/Online/HTTPChunkInstaller/Source/Public/ChunkSetup.h
deleted file mode 100644
diff --git a/Engine/Plugins/Online/HTTPChunkInstaller/Source/Public/HTTPChunkInstaller.h b/Engine/Plugins/Online/HTTPChunkInstaller/Source/Public/HTTPChunkInstaller.h
deleted file mode 100644
diff --git a/Engine/Plugins/Online/HTTPChunkInstaller/Source/Public/HTTPChunkInstallerPrivatePCH.h b/Engine/Plugins/Online/HTTPChunkInstaller/Source/Public/HTTPChunkInstallerPrivatePCH.h
deleted file mode 100644
diff --git a/Engine/Plugins/Online/HTTPChunkInstaller/Source/Public/HTTPOnlineTitleFile.h b/Engine/Plugins/Online/HTTPChunkInstaller/Source/Public/HTTPOnlineTitleFile.h
deleted file mode 100644
diff --git a/Engine/Plugins/Online/HTTPChunkInstaller/Source/Public/LocalTitleFile.h b/Engine/Plugins/Online/HTTPChunkInstaller/Source/Public/LocalTitleFile.h
deleted file mode 100644
diff --git a/Engine/Plugins/Runtime/HTTPChunkInstaller/Config/Localization/HTTPChunkInstaller.ini b/Engine/Plugins/Runtime/HTTPChunkInstaller/Config/Localization/HTTPChunkInstaller.ini
new file mode 100644
--- /dev/null
+++ b/Engine/Plugins/Runtime/HTTPChunkInstaller/Config/Localization/HTTPChunkInstaller.ini
@@ -0,0 +1,55 @@
+[CommonSettings]
+SourcePath=Plugins/Online/HTTPChunkInstaller/Content/Localization/HTTPChunkInstaller
+DestinationPath=Plugins/Online/HTTPChunkInstaller/Content/Localization/HTTPChunkInstaller
+ManifestName=HTTPChunkInstaller.manifest
+ArchiveName=HTTPChunkInstaller.archive
+PortableObjectName=HTTPChunkInstaller.po
+ResourceName=HTTPChunkInstaller.locres
+NativeCulture=en
+CulturesToGenerate=en
+CulturesToGenerate=fr
+CulturesToGenerate=de
+CulturesToGenerate=es
+CulturesToGenerate=es-419
+CulturesToGenerate=it
+CulturesToGenerate=ja
+CulturesToGenerate=ko
+CulturesToGenerate=ru
+CulturesToGenerate=zh-Hans
+CulturesToGenerate=ar
+CulturesToGenerate=pl
+CulturesToGenerate=pt-BR
+
+;Gather text from source code
+[GatherTextStep0]
+CommandletClass=GatherTextFromSource
+SearchDirectoryPaths=Plugins/Online/HTTPChunkInstaller/Source/
+FileNameFilters=*.cpp
+FileNameFilters=*.h
+FileNameFilters=*.c
+FileNameFilters=*.inl
+FileNameFilters=*.mm
+FileNameFilters=*.ini
+ShouldGatherFromEditorOnlyData=false
+
+;Write Manifest
+[GatherTextStep1]
+CommandletClass=GenerateGatherManifest
+
+;Write Archives
+[GatherTextStep2]
+CommandletClass=GenerateGatherArchive
+
+;Import localized PO files
+[GatherTextStep3]
+CommandletClass=InternationalizationExport
+bImportLoc=true
+
+;Write Localized Text Resource
+[GatherTextStep4]
+CommandletClass=GenerateTextLocalizationResource
+
+;Export  PO files
+[GatherTextStep5]
+CommandletClass=InternationalizationExport
+bExportLoc=true
\ No newline at end of file
diff --git a/Engine/Plugins/Runtime/HTTPChunkInstaller/HTTPChunkInstaller.uplugin b/Engine/Plugins/Runtime/HTTPChunkInstaller/HTTPChunkInstaller.uplugin
new file mode 100644
--- /dev/null
+++ b/Engine/Plugins/Runtime/HTTPChunkInstaller/HTTPChunkInstaller.uplugin
@@ -0,0 +1,26 @@
+{
+	"FileVersion" : 3,
+	
+	"FriendlyName" : "HTTP Chunk Installer",
+	"Version" : 1,
+	"VersionName" : "1.0",
+	"Description" : "Implements a streaming install client",
+	"Category" : "Online Platform",
+	"CreatedBy" : "Epic Games, Inc.",
+	"CreatedByURL" : "http://epicgames.com",
+	"EnabledByDefault" : false,
+
+    "Modules": [
+        {
+            "Name": "HTTPChunkInstaller",
+            "Type": "Runtime",
+            "LoadingPhase" : "PostConfigInit"
+        }
+    ],
+	"LocalizationTargets": [
+		{
+			"Name": "HTTPChunkInstaller",
+			"LoadingPolicy": "Always"
+		}
+	]
+}
\ No newline at end of file
diff --git a/Engine/Plugins/Runtime/HTTPChunkInstaller/Source/HTTPChunkInstaller.Build.cs b/Engine/Plugins/Runtime/HTTPChunkInstaller/Source/HTTPChunkInstaller.Build.cs
new file mode 100644
--- /dev/null
+++ b/Engine/Plugins/Runtime/HTTPChunkInstaller/Source/HTTPChunkInstaller.Build.cs
@@ -0,0 +1,26 @@
+// Copyright 1998-2016 Epic Games, Inc. All Rights Reserved.
+
+using UnrealBuildTool;
+
+public class HTTPChunkInstaller : ModuleRules
+{
+	public HTTPChunkInstaller(ReadOnlyTargetRules Target)
+		: base(Target)
+	{
+        PrivateDependencyModuleNames.AddRange(
+            new string[] {
+                "BuildPatchServices",
+                "Core",
+                "Engine",
+                "Http",
+                "Json",
+                "PakFile",
+            }
+            );
+
+ //       if (Target.Platform != UnrealTargetPlatform.Win64 && Target.Platform != UnrealTargetPlatform.Win32 && Target.Platform != UnrealTargetPlatform.IOS)
+        {
+            PrecompileForTargets = PrecompileTargetsType.None;
+        }
+    }
+}
\ No newline at end of file
diff --git a/Engine/Plugins/Runtime/HTTPChunkInstaller/Source/Private/HTTPChunkInstaller.cpp b/Engine/Plugins/Runtime/HTTPChunkInstaller/Source/Private/HTTPChunkInstaller.cpp
new file mode 100644
--- /dev/null
+++ b/Engine/Plugins/Runtime/HTTPChunkInstaller/Source/Private/HTTPChunkInstaller.cpp
@@ -0,0 +1,1746 @@
+// Copyright 1998-2016 Epic Games, Inc. All Rights Reserved.
+
+#ifndef PLACEHOLDER_DLC_WORK
+#	define PLACEHOLDER_DLC_WORK 0
+#endif
+
+#include "HTTPChunkInstaller.h"
+#include "HTTPChunkInstallerLog.h"
+#include "ChunkInstall.h"
+#include "UniquePtr.h"
+#include "LocalTitleFile.h"
+#include "SecureHash.h"
+#include "IHttpRequest.h"
+#include "IHttpResponse.h"
+#include "HTTPModule.h"
+#include "JsonObject.h"
+#include "JsonReader.h"
+#include "JsonSerializer.h"
+#include "ThreadSafeCounter64.h"
+#include "ConfigCacheIni.h"
+#include "FileHelper.h"
+#include "RunnableThread.h"
+#include "CommandLine.h"
+#include "ModuleManager.h"
+
+#define LOCTEXT_NAMESPACE "HTTPChunkInstaller"
+
+// helper to grab the installer service
+static IBuildPatchServicesModule* GetBuildPatchServices()
+{
+	static IBuildPatchServicesModule* BuildPatchServices = nullptr;
+	if (BuildPatchServices == nullptr)
+	{
+		BuildPatchServices = &FModuleManager::LoadModuleChecked<IBuildPatchServicesModule>(TEXT("BuildPatchServices"));
+	}
+	return BuildPatchServices;
+}
+
+const FName NAME_SHA1(TEXT("SHA1"));
+const FName NAME_SHA256(TEXT("SHA256"));
+
+class FTitleFileHttpAsyncLoadAndVerify :
+	public FNonAbandonableTask
+{
+public:
+	/** File data loaded for the async read */
+	TArray<uint8> FileData;
+	/** Amount of data read from the file to be owned/referenced by the game thread */
+	FThreadSafeCounter64* BytesRead;
+	/** The original name of the file being read */
+	const FString OriginalFileName;
+	/** The name of the file being read off of disk */
+	const FString FileName;
+	/** The hash value the backend said it should have */
+	const FString ExpectedHash;
+	/** The hash type SHA1 or SHA256 right now */
+	const FName HashType;
+	/** Whether the hashes matched */
+	bool bHashesMatched;
+
+	/** Initializes the variables needed to load and verify the data */
+	FTitleFileHttpAsyncLoadAndVerify(const FString& InOriginalFileName, const FString& InFileName, const FString& InExpectedHash, const FName InHashType, FThreadSafeCounter64* InBytesReadCounter) :
+		BytesRead(InBytesReadCounter),
+		OriginalFileName(InOriginalFileName),
+		FileName(InFileName),
+		ExpectedHash(InExpectedHash),
+		HashType(InHashType),
+		bHashesMatched(false)
+	{
+
+	}
+
+	/**
+	* Loads and hashes the file data. Empties the data if the hash check fails
+	*/
+	void DoWork()
+	{
+		// load file from disk
+		bool bLoadedFile = false;
+
+		FArchive* Reader = IFileManager::Get().CreateFileReader(*FileName, FILEREAD_Silent);
+		if (Reader)
+		{
+			int64 SizeToRead = Reader->TotalSize();
+			FileData.Reset(SizeToRead);
+			FileData.AddUninitialized(SizeToRead);
+
+			uint8* FileDataPtr = FileData.GetData();
+
+			static const int64 ChunkSize = 100 * 1024;
+
+			int64 TotalBytesRead = 0;
+			while (SizeToRead > 0)
+			{
+				int64 Val = FMath::Min(SizeToRead, ChunkSize);
+				Reader->Serialize(FileDataPtr + TotalBytesRead, Val);
+				BytesRead->Add(Val);
+				TotalBytesRead += Val;
+				SizeToRead -= Val;
+			}
+
+			ensure(SizeToRead == 0 && Reader->TotalSize() == TotalBytesRead);
+			bLoadedFile = Reader->Close();
+			delete Reader;
+		}
+
+		// verify hash of file if it exists
+		if (bLoadedFile)
+		{
+			UE_LOG(LogHTTPChunkInstaller, Verbose, TEXT("ReadFile request. Local file read from cache =%s"), *FileName);
+			if (HashType == NAME_SHA1)
+			{
+				bHashesMatched = IsValidSHA1(ExpectedHash, FileData);
+			}
+			else if (HashType == NAME_SHA256)
+			{
+				bHashesMatched = IsValidSHA256(ExpectedHash, FileData);
+			}
+		}
+		else
+		{
+			UE_LOG(LogHTTPChunkInstaller, Verbose, TEXT("Local file (%s) not cached locally"), *FileName);
+		}
+		if (!bHashesMatched)
+		{
+			// Empty local that was loaded
+			FileData.Empty();
+		}
+	}
+
+	FORCEINLINE TStatId GetStatId() const
+	{
+		RETURN_QUICK_DECLARE_CYCLE_STAT(FTitleFileHttpAsyncLoadAndVerify, STATGROUP_ThreadPoolAsyncTasks);
+	}
+
+private:
+	/** Validates that a buffer matches the same signature as was specified */
+	bool IsValidSHA1(const FString& Hash, const TArray<uint8>& Source) const
+	{
+		uint8 LocalHash[20];
+		FSHA1::HashBuffer(Source.GetData(), Source.Num(), LocalHash);
+		// concatenate 20 bye SHA1 hash to string
+		FString LocalHashStr;
+		for (int Idx = 0; Idx < 20; Idx++)
+		{
+			LocalHashStr += FString::Printf(TEXT("%02x"), LocalHash[Idx]);
+		}
+		return Hash == LocalHashStr;
+	}
+	bool IsValidSHA256(const FString& Hash, const TArray<uint8>& Source) const
+	{
+		FSHA256Signature Signature;
+		if (FPlatformMisc::GetSHA256Signature(Source.GetData(), Source.Num(), Signature))
+		{
+			return Signature.ToString() == Hash;
+		}
+		return false;
+	}
+};
+
+// helper class for an HTTP Online Title File
+class FOnlineTitleFileHttp : public ICloudTitleFile
+{
+public:
+	/**
+	* Constructor
+	*
+	* @param InSubsystem mcp subsystem being used
+	*/
+	FOnlineTitleFileHttp(const FString& InBaseUrl)
+		: BaseUrl(InBaseUrl)
+		, EnumerateFilesUrl(TEXT(""))
+	{
+		GConfig->GetString(TEXT("HTTPOnlineTitleFile"), TEXT("BaseUrl"), BaseUrl, GEngineIni);
+		GConfig->GetString(TEXT("HTTPOnlineTitleFile"), TEXT("EnumerateFilesUrl"), EnumerateFilesUrl, GEngineIni);
+		bCacheFiles = true;
+		bPlatformSupportsSHA256 = false;
+	}
+
+	/**
+	* Destructor
+	*/
+	virtual ~FOnlineTitleFileHttp()
+	{}
+
+	/**
+	* Copies the file data into the specified buffer for the specified file
+	*
+	* @param FileName the name of the file to read
+	* @param FileContents the out buffer to copy the data into
+	*
+	* @return true if the data was copied, false otherwise
+	*/
+	virtual bool GetFileContents(const FString& FileName, TArray<uint8>& FileContents) override
+	{
+		const TArray<FCloudEntry>* FilesPtr = &Files;
+		if (FilesPtr != NULL)
+		{
+			for (TArray<FCloudEntry>::TConstIterator It(*FilesPtr); It; ++It)
+			{
+				if (It->FileName == FileName)
+				{
+					FileContents = It->Data;
+					return true;
+				}
+			}
+		}
+		return false;
+	}
+
+	/**
+	* Empties the set of downloaded files if possible (no async tasks outstanding)
+	*
+	* @return true if they could be deleted, false if they could not
+	*/
+	virtual bool ClearFiles() override
+	{
+		for (int Idx = 0; Idx < Files.Num(); Idx++)
+		{
+			if (Files[Idx].AsyncState == ECloudAsyncTaskState::InProgress)
+			{
+				UE_LOG(LogHTTPChunkInstaller, Warning, TEXT("Cant clear files. Pending file op for %s"), *Files[Idx].FileName);
+				return false;
+			}
+		}
+		// remove all cached file entries
+		Files.Empty();
+		return true;
+	}
+
+	/**
+	* Empties the cached data for this file if it is not being downloaded currently
+	*
+	* @param FileName the name of the file to remove from the cache
+	*
+	* @return true if it could be deleted, false if it could not
+	*/
+	virtual bool ClearFile(const FString& FileName) override
+	{
+		for (int Idx = 0; Idx < Files.Num(); Idx++)
+		{
+			if (Files[Idx].FileName == FileName)
+			{
+				if (Files[Idx].AsyncState == ECloudAsyncTaskState::InProgress)
+				{
+					UE_LOG(LogHTTPChunkInstaller, Warning, TEXT("Cant clear file. Pending file op for %s"), *Files[Idx].FileName);
+					return false;
+				}
+				else
+				{
+					Files.RemoveAt(Idx);
+					return true;
+				}
+			}
+		}
+		return false;
+	}
+
+	/**
+	* Delete cached files on disk
+	*
+	* @param bSkipEnumerated if true then only non-enumerated files are deleted
+	*/
+	virtual void DeleteCachedFiles(bool bSkipEnumerated) override
+	{
+		TArray<FString> CachedFiles;
+		IFileManager::Get().FindFiles(CachedFiles, *(GetLocalCachePath() / TEXT("*")), true, false);
+
+		for (auto CachedFile : CachedFiles)
+		{
+			bool bSkip = bSkipEnumerated && GetCloudFileHeader(CachedFile);
+			if (!bSkip)
+			{
+				IFileManager::Get().Delete(*GetLocalFilePath(CachedFile), false, true);
+			}
+		}
+	}
+
+	/**
+	* Requests a list of available files from the network store
+	*
+	* @return true if the request has started, false if not
+	*/
+	virtual bool EnumerateFiles(const FCloudPagedQuery& Page = FCloudPagedQuery()) override
+	{
+		FString ErrorStr;
+		bool bStarted = false;
+
+		// Make sure an enumeration request  is not currently pending
+		if (EnumerateFilesRequests.Num() > 0)
+		{
+			ErrorStr = TEXT("Request already in progress.");
+		}
+		else
+		{
+			// Create the Http request and add to pending request list
+			TSharedRef<IHttpRequest> HttpRequest = FHttpModule::Get().CreateRequest();
+			EnumerateFilesRequests.Add(HttpRequest, Page);
+
+			HttpRequest->OnProcessRequestComplete().BindRaw(this, &FOnlineTitleFileHttp::EnumerateFiles_HttpRequestComplete);
+			HttpRequest->SetURL(GetBaseUrl()+EnumerateFilesUrl+TEXT("/Master.manifest"));
+			HttpRequest->SetVerb(TEXT("GET"));
+			bStarted = HttpRequest->ProcessRequest();
+		}
+		if (!bStarted)
+		{
+			UE_LOG(LogHTTPChunkInstaller, Warning, TEXT("EnumerateFiles request failed. %s"), *ErrorStr);
+			TriggerOnEnumerateFilesCompleteDelegates(false);
+		}
+		return bStarted;
+	}
+
+	/**
+	* Returns the list of files that was returned by the network store
+	*
+	* @param Files out array of file metadata
+	*
+	*/
+	virtual void GetFileList(TArray<FCloudHeader>& OutFiles) override
+	{
+		TArray<FCloudHeader>* FilesPtr = &FileHeaders;
+		if (FilesPtr != NULL)
+		{
+			OutFiles = *FilesPtr;
+		}
+		else
+		{
+			OutFiles.Empty();
+		}
+	}
+
+	/**
+	* Starts an asynchronous read of the specified file from the network platform's file store
+	*
+	* @param FileToRead the name of the file to read
+	*
+	* @return true if the calls starts successfully, false otherwise
+	*/
+	virtual bool ReadFile(const FString& FileName) override
+	{
+		bool bStarted = false;
+
+		FCloudHeader* CloudFileHeader = GetCloudFileHeader(FileName);
+
+		// Make sure valid filename was specified3
+		if (FileName.IsEmpty() || FileName.Contains(TEXT(" ")))
+		{
+			UE_LOG(LogHTTPChunkInstaller, Warning, TEXT("Invalid filename filename=%s"), *FileName);
+			TriggerOnReadFileCompleteDelegates(false, FileName);
+			return false;
+		}
+
+		// Make sure a file request for this file is not currently pending
+		for (const auto& Pair : FileRequests)
+		{
+			if (Pair.Value == FPendingFileRequest(FileName))
+			{
+				UE_LOG(LogHTTPChunkInstaller, Verbose, TEXT("ReadFileRemote is already in progress for (%s)"), *FileName);
+				return true;
+			}
+		}
+
+		FCloudEntry* CloudFile = GetCloudFile(FileName, true);
+		if (CloudFile->AsyncState == ECloudAsyncTaskState::InProgress)
+		{
+			UE_LOG(LogHTTPChunkInstaller, Verbose, TEXT("ReadFile is already in progress for (%s)"), *FileName);
+			return true;
+		}
+
+		if (bCacheFiles)
+		{
+			// Try to read this from the cache if possible
+			bStarted = StartReadFileLocal(FileName);
+		}
+		if (!bStarted)
+		{
+			// Failed locally (means not on disk) so fetch from server
+			bStarted = ReadFileRemote(FileName);
+		}
+
+		if (!bStarted || CloudFile->AsyncState == ECloudAsyncTaskState::Failed)
+		{
+			UE_LOG(LogHTTPChunkInstaller, Warning, TEXT("ReadFile request failed for file (%s)"), *FileName);
+			TriggerOnReadFileCompleteDelegates(false, FileName);
+		}
+		else if (CloudFile->AsyncState == ECloudAsyncTaskState::Done)
+		{
+			TriggerOnReadFileCompleteDelegates(true, FileName);
+		}
+		return bStarted;
+	}
+
+	/** Used to check that async tasks have completed and can be completed */
+	virtual void Tick(float DeltaTime)
+	{
+		TArray<int32> ItemsToRemove;
+		ItemsToRemove.Reserve(AsyncLocalReads.Num());
+
+		// Check for any completed tasks
+		for (int32 TaskIdx = 0; TaskIdx < AsyncLocalReads.Num(); TaskIdx++)
+		{
+			FTitleAsyncReadData& Task = AsyncLocalReads[TaskIdx];
+			if (Task.AsyncTask->IsDone())
+			{
+				FinishReadFileLocal(Task.AsyncTask->GetTask());
+				ItemsToRemove.Add(TaskIdx);
+				UE_LOG(LogHTTPChunkInstaller, VeryVerbose, TEXT("Title Task Complete: %s"), *Task.Filename);
+			}
+			else
+			{
+				const int64 NewValue = Task.BytesRead.GetValue();
+				if (NewValue != Task.LastBytesRead)
+				{
+					Task.LastBytesRead = NewValue;
+					TriggerOnReadFileProgressDelegates(Task.Filename, NewValue);
+				}
+			}
+		}
+
+		// Clean up any tasks that were completed
+		for (int32 ItemIdx = ItemsToRemove.Num() - 1; ItemIdx >= 0; ItemIdx--)
+		{
+			const int32 TaskIdx = ItemsToRemove[ItemIdx];
+			FTitleAsyncReadData& TaskToDelete = AsyncLocalReads[TaskIdx];
+			UE_LOG(LogHTTPChunkInstaller, VeryVerbose, TEXT("Title Task Removal: %s read: %d"), *TaskToDelete.Filename, TaskToDelete.BytesRead.GetValue());
+			delete TaskToDelete.AsyncTask;
+			AsyncLocalReads.RemoveAtSwap(TaskIdx);
+		}
+	}
+
+	void Shutdown()
+	{
+
+	}
+
+private:
+
+	/** Reads the file from the local cache if it can. This is async */
+	bool StartReadFileLocal(const FString& FileName)
+	{
+		UE_LOG(LogHTTPChunkInstaller, VeryVerbose, TEXT("StartReadFile %s"), *FileName);
+		bool bStarted = false;
+		FCloudHeader* CloudFileHeader = GetCloudFileHeader(FileName);
+		if (CloudFileHeader != nullptr)
+		{
+			// Mark file entry as in progress
+			FCloudEntry* CloudFile = GetCloudFile(FileName, true);
+			CloudFile->AsyncState = ECloudAsyncTaskState::InProgress;
+			if (CloudFileHeader->Hash.IsEmpty())
+			{
+				UE_LOG(LogHTTPChunkInstaller, Warning, TEXT("Requested file (%s) is missing a hash, so can't be verified"), *FileName);
+			}
+			FTitleAsyncReadData* NewItem = new FTitleAsyncReadData();
+			NewItem->Filename = FileName;
+
+			// Create the async task and start it
+			NewItem->AsyncTask = new FAsyncTask<FTitleFileHttpAsyncLoadAndVerify>(FileName,
+				GetLocalFilePath(FileName), CloudFileHeader->Hash, CloudFileHeader->HashType, &NewItem->BytesRead);
+
+			AsyncLocalReads.Add(NewItem);
+			NewItem->AsyncTask->StartBackgroundTask();
+			bStarted = true;
+		}
+		return bStarted;
+	}
+	/** Completes the async operation of the local file read */
+	void FinishReadFileLocal(FTitleFileHttpAsyncLoadAndVerify& AsyncLoad)
+	{
+		UE_LOG(LogHTTPChunkInstaller, VeryVerbose, TEXT("FinishReadFileLocal %s"), *AsyncLoad.OriginalFileName);
+		FCloudHeader* CloudFileHeader = GetCloudFileHeader(AsyncLoad.OriginalFileName);
+		FCloudEntry* CloudFile = GetCloudFile(AsyncLoad.OriginalFileName, true);
+		if (CloudFileHeader != nullptr && CloudFile != nullptr)
+		{
+			// if hash matches then just use the local file
+			if (AsyncLoad.bHashesMatched)
+			{
+				UE_LOG(LogHTTPChunkInstaller, Verbose, TEXT("Local file hash matches cloud header. No need to download for filename=%s"), *AsyncLoad.OriginalFileName);
+				CloudFile->Data = AsyncLoad.FileData;
+				CloudFile->AsyncState = ECloudAsyncTaskState::Done;
+				TriggerOnReadFileProgressDelegates(AsyncLoad.OriginalFileName, static_cast<uint64>(AsyncLoad.BytesRead->GetValue()));
+				TriggerOnReadFileCompleteDelegates(true, AsyncLoad.OriginalFileName);
+			}
+			else
+			{
+				// Request it from server
+				ReadFileRemote(AsyncLoad.OriginalFileName);
+			}
+		}
+		else
+		{
+			UE_LOG(LogHTTPChunkInstaller, Warning, TEXT("ReadFile request failed for file (%s)"), *AsyncLoad.OriginalFileName);
+			TriggerOnReadFileCompleteDelegates(false, AsyncLoad.OriginalFileName);
+		}
+	}
+	/** Requests the file from MCP. This is async */
+	bool ReadFileRemote(const FString& FileName)
+	{
+		UE_LOG(LogHTTPChunkInstaller, VeryVerbose, TEXT("ReadFileRemote %s"), *FileName);
+
+		bool bStarted = false;
+		FCloudHeader* CloudFileHeader = GetCloudFileHeader(FileName);
+		if (CloudFileHeader != nullptr)
+		{
+			FCloudEntry* CloudFile = GetCloudFile(FileName, true);
+			CloudFile->Data.Empty();
+			CloudFile->AsyncState = ECloudAsyncTaskState::InProgress;
+
+			// Create the Http request and add to pending request list
+			TSharedRef<IHttpRequest> HttpRequest = FHttpModule::Get().CreateRequest();
+			FileRequests.Add(HttpRequest, FPendingFileRequest(FileName));
+			FileProgressRequestsMap.Add(HttpRequest, FPendingFileRequest(FileName));
+
+			HttpRequest->OnProcessRequestComplete().BindRaw(this, &FOnlineTitleFileHttp::ReadFile_HttpRequestComplete);
+			HttpRequest->OnRequestProgress().BindRaw(this, &FOnlineTitleFileHttp::ReadFile_HttpRequestProgress);
+			FString RequestUrl;
+			// Grab the file from the specified URL if that was set, otherwise use the old method that hits the game service
+			if (CloudFileHeader != nullptr && !CloudFileHeader->URL.IsEmpty())
+			{
+				RequestUrl = CloudFileHeader->URL;
+			}
+			else
+			{
+				RequestUrl = GetBaseUrl() + FileName;
+			}
+			HttpRequest->SetURL(RequestUrl);
+			HttpRequest->SetVerb(TEXT("GET"));
+			bStarted = HttpRequest->ProcessRequest();
+
+			if (!bStarted)
+			{
+				UE_LOG(LogHTTPChunkInstaller, Error, TEXT("Unable to start the HTTP request to fetch file (%s)"), *FileName);
+			}
+		}
+		else
+		{
+			UE_LOG(LogHTTPChunkInstaller, Error, TEXT("No cloud file header entry for filename=%s."), *FileName);
+		}
+		return bStarted;
+	}
+
+	/**
+	* Delegate called when a Http request completes for enumerating list of file headers
+	*/
+	void EnumerateFiles_HttpRequestComplete(FHttpRequestPtr HttpRequest, FHttpResponsePtr HttpResponse, bool bSucceeded)
+	{
+		const FCloudPagedQuery& PendingOp = EnumerateFilesRequests.FindRef(HttpRequest);
+		EnumerateFilesRequests.Remove(HttpRequest);
+
+		bool bResult = false;
+		FString ResponseStr, ErrorStr;
+
+		if (HttpResponse.IsValid() && EHttpResponseCodes::IsOk(HttpResponse->GetResponseCode()))
+		{
+			ResponseStr = HttpResponse->GetContentAsString();
+			UE_LOG(LogHTTPChunkInstaller, Verbose, TEXT("EnumerateFiles request complete. url=%s code=%d response=%s"),
+				*HttpRequest->GetURL(), HttpResponse->GetResponseCode(), *ResponseStr);
+
+			if (PendingOp.Start == 0)
+			{
+				FileHeaders.Empty();
+			}
+
+			// parse the html for the file list
+			if (ResponseStr.StartsWith(TEXT("<!DOCTYPE")))
+			{
+				TArray<FString> Lines;
+				ResponseStr.ParseIntoArrayLines(Lines);
+				for (int Index = 0; Index < Lines.Num(); ++Index)
+				{
+					if (Lines[Index].StartsWith(TEXT("<li>")))
+					{
+						TArray<FString> Elements;
+						Lines[Index].ParseIntoArray(Elements, TEXT(">"));
+						if (!Elements[2].StartsWith(TEXT("Chunks")))
+						{
+							FString File = Elements[2].Replace(TEXT("</a"), TEXT(""));
+							FCloudHeader FileHeader;
+							FileHeader.DLName = File;
+							FileHeader.FileName = File;
+							FileHeader.URL = GetBaseUrl() + EnumerateFilesUrl + TEXT("/") + File;
+							FileHeader.Hash.Empty();
+							FileHeader.FileSize = 0;
+							FileHeaders.Add(FileHeader);
+						}
+					}
+				}
+				bResult = true;
+			}
+			else
+			{
+				// Create the Json parser
+				TSharedPtr<FJsonObject> JsonObject;
+				TSharedRef<TJsonReader<> > JsonReader = TJsonReaderFactory<>::Create(ResponseStr);
+
+				if (FJsonSerializer::Deserialize(JsonReader, JsonObject) &&
+					JsonObject.IsValid())
+				{
+					// Parse the array of file headers
+					TArray<TSharedPtr<FJsonValue> > JsonFileHeaders = JsonObject->GetArrayField(TEXT("files"));
+					for (TArray<TSharedPtr<FJsonValue> >::TConstIterator It(JsonFileHeaders); It; ++It)
+					{
+						TSharedPtr<FJsonObject> JsonFileHeader = (*It)->AsObject();
+						if (JsonFileHeader.IsValid())
+						{
+							FCloudHeader FileHeader;
+							if (JsonFileHeader->HasField(TEXT("hash")))
+							{
+								FileHeader.Hash = JsonFileHeader->GetStringField(TEXT("hash"));
+								FileHeader.HashType = FileHeader.Hash.IsEmpty() ? NAME_None : NAME_SHA1;
+							}
+							// This one takes priority over the old SHA1 hash if present (requires platform support)
+							if (bPlatformSupportsSHA256 && JsonFileHeader->HasField(TEXT("hash256")))
+							{
+								FString Hash256 = JsonFileHeader->GetStringField(TEXT("hash256"));
+								if (!Hash256.IsEmpty())
+								{
+									FileHeader.Hash = Hash256;
+									FileHeader.HashType = FileHeader.Hash.IsEmpty() ? NAME_None : NAME_SHA256;
+								}
+							}
+							if (JsonFileHeader->HasField(TEXT("uniqueFilename")))
+							{
+								FileHeader.DLName = JsonFileHeader->GetStringField(TEXT("uniqueFilename"));
+							}
+							if (JsonFileHeader->HasField(TEXT("filename")))
+							{
+								FileHeader.FileName = JsonFileHeader->GetStringField(TEXT("filename"));
+							}
+							if (JsonFileHeader->HasField(TEXT("length")))
+							{
+								FileHeader.FileSize = FMath::TruncToInt(JsonFileHeader->GetNumberField(TEXT("length")));
+							}
+							if (JsonFileHeader->HasField(TEXT("URL")))
+							{
+								FileHeader.URL = GetBaseUrl() + EnumerateFilesUrl + TEXT("/") + JsonFileHeader->GetStringField(TEXT("URL"));
+							}
+
+							if (FileHeader.FileName.IsEmpty())
+							{
+								FileHeader.FileName = FileHeader.DLName;
+							}
+
+							if (FileHeader.Hash.IsEmpty() ||
+								(FileHeader.DLName.IsEmpty() && FileHeader.URL.IsEmpty()) ||
+								FileHeader.HashType == NAME_None)
+							{
+								UE_LOG(LogHTTPChunkInstaller, Warning, TEXT("Invalid file entry hash=%s hashType=%s dlname=%s filename=%s URL=%s"),
+									*FileHeader.Hash,
+									*FileHeader.HashType.ToString(),
+									*FileHeader.DLName,
+									*FileHeader.FileName,
+									*FileHeader.URL);
+							}
+							else
+							{
+								int32 FoundIdx = INDEX_NONE;
+								for (int32 Idx = 0; Idx < FileHeaders.Num(); Idx++)
+								{
+									const FCloudHeader& ExistingFile = FileHeaders[Idx];
+									if (ExistingFile.DLName == FileHeader.DLName)
+									{
+										FoundIdx = Idx;
+										break;
+									}
+								}
+								if (FoundIdx != INDEX_NONE)
+								{
+									FileHeaders[FoundIdx] = FileHeader;
+								}
+								else
+								{
+									FileHeaders.Add(FileHeader);
+								}
+							}
+						}
+					}
+				}
+				bResult = true;
+			}
+		}
+		else
+		{
+			if (HttpResponse.IsValid())
+			{
+				ErrorStr = FText::Format(LOCTEXT("HttpResponse", "HTTP {0} response from {1}"),
+					FText::AsNumber(HttpResponse->GetResponseCode()),
+					FText::FromString(HttpResponse->GetURL())).ToString();
+			}
+			else
+			{
+				ErrorStr = FText::Format(LOCTEXT("HttpResponse", "Connection to {0} failed"), FText::FromString(HttpRequest->GetURL())).ToString();
+			}
+		}
+
+		if (!ErrorStr.IsEmpty())
+		{
+			UE_LOG(LogHTTPChunkInstaller, Warning, TEXT("EnumerateFiles request failed. %s"), *ErrorStr);
+		}
+		else
+		{
+			// Everything went ok, so we can remove any cached files that are not in the current list
+			DeleteCachedFiles(true);
+		}
+
+		TriggerOnEnumerateFilesCompleteDelegates(bResult);
+	}
+
+	/**
+	* Delegate called when a Http request completes for reading a cloud file
+	*/
+	void ReadFile_HttpRequestComplete(FHttpRequestPtr HttpRequest, FHttpResponsePtr HttpResponse, bool bSucceeded)
+	{
+		bool bResult = false;
+		FString ResponseStr, ErrorStr;
+
+		// should have a pending Http request
+		FPendingFileRequest PendingRequest = FileRequests.FindChecked(HttpRequest);
+		FileRequests.Remove(HttpRequest);
+		// remove from progress updates
+		FileProgressRequestsMap.Remove(HttpRequest);
+		HttpRequest->OnRequestProgress().Unbind();
+
+		// Cloud file being operated on
+		FCloudEntry* CloudFile = GetCloudFile(PendingRequest.FileName, true);
+		CloudFile->AsyncState = ECloudAsyncTaskState::Failed;
+		CloudFile->Data.Empty();
+
+		if (HttpResponse.IsValid() && EHttpResponseCodes::IsOk(HttpResponse->GetResponseCode()))
+		{
+			UE_LOG(LogHTTPChunkInstaller, Verbose, TEXT("ReadFile request complete. url=%s code=%d"),
+				*HttpRequest->GetURL(), HttpResponse->GetResponseCode());
+
+			// update the memory copy of the file with data that was just downloaded
+			CloudFile->AsyncState = ECloudAsyncTaskState::Done;
+			CloudFile->Data = HttpResponse->GetContent();
+
+			if (bCacheFiles)
+			{
+				// cache to disk on successful download
+				SaveCloudFileToDisk(CloudFile->FileName, CloudFile->Data);
+			}
+
+			bResult = true;
+		}
+		else
+		{
+			if (HttpResponse.IsValid())
+			{
+				ErrorStr = FText::Format(LOCTEXT("HttpResponse", "HTTP {0} response from {1}"),
+					FText::AsNumber(HttpResponse->GetResponseCode()),
+					FText::FromString(HttpResponse->GetURL())).ToString();
+			}
+			else
+			{
+				ErrorStr = FText::Format(LOCTEXT("HttpResponse", "Connection to {0} failed"), FText::FromString(HttpRequest->GetURL())).ToString();
+			}
+		}
+
+		if (!ErrorStr.IsEmpty())
+		{
+			UE_LOG(LogHTTPChunkInstaller, Warning, TEXT("EnumerateFiles request failed. %s"), *ErrorStr);
+		}
+		TriggerOnReadFileCompleteDelegates(bResult, PendingRequest.FileName);
+	}
+
+	/**
+	* Delegate called as a Http request progresses for reading a cloud file
+	*/
+	void ReadFile_HttpRequestProgress(FHttpRequestPtr HttpRequest, int32 BytesSent, int32 BytesReceived)
+	{
+		FPendingFileRequest PendingRequest = FileProgressRequestsMap.FindChecked(HttpRequest);
+		// Just forward this to anyone that is listening
+		TriggerOnReadFileProgressDelegates(PendingRequest.FileName, BytesReceived);
+	}
+
+	/**
+	* Find/create cloud file entry
+	*
+	* @param FileName cached file entry to find
+	* @param bCreateIfMissing create the file entry if not found
+	*
+	* @return cached cloud file entry
+	*/
+	FCloudEntry* GetCloudFile(const FString& FileName, bool bCreateIfMissing)
+	{
+		FCloudEntry* CloudFile = NULL;
+		for (int Idx = 0; Idx < Files.Num(); Idx++)
+		{
+			if (Files[Idx].FileName == FileName)
+			{
+				CloudFile = &Files[Idx];
+				break;
+			}
+		}
+		if (CloudFile == NULL && bCreateIfMissing)
+		{
+			CloudFile = new(Files)FCloudEntry(FileName);
+		}
+		return CloudFile;
+	}
+
+	/**
+	* Find cloud file header entry
+	*
+	* @param FileName cached file entry to find
+	*
+	* @return cached cloud file header entry
+	*/
+	FCloudHeader* GetCloudFileHeader(const FString& FileName)
+	{
+		FCloudHeader* CloudFileHeader = NULL;
+
+		for (int Idx = 0; Idx < FileHeaders.Num(); Idx++)
+		{
+			if (FileHeaders[Idx].DLName == FileName)
+			{
+				CloudFileHeader = &FileHeaders[Idx];
+				break;
+			}
+		}
+
+		return CloudFileHeader;
+	}
+
+	/**
+	* Converts filename into a local file cache path
+	*
+	* @param FileName name of file being loaded
+	*
+	* @return unreal file path to be used by file manager
+	*/
+	FString GetLocalFilePath(const FString& FileName)
+	{
+		return GetLocalCachePath() + FileName;
+	}
+
+	/**
+	* @return full path to cache directory
+	*/
+	FString GetLocalCachePath()
+	{
+		return FPaths::GamePersistentDownloadDir() / TEXT("EMS/");
+	}
+
+	/**
+	* Save a file from a given user to disk
+	*
+	* @param FileName name of file being saved
+	* @param Data data to write to disk
+	*/
+	void SaveCloudFileToDisk(const FString& Filename, const TArray<uint8>& Data)
+	{
+		// save local disk copy as well
+		FString LocalFilePath = GetLocalFilePath(Filename);
+		bool bSavedLocal = FFileHelper::SaveArrayToFile(Data, *LocalFilePath);
+		if (bSavedLocal)
+		{
+			UE_LOG(LogHTTPChunkInstaller, Verbose, TEXT("WriteUserFile request complete. Local file cache updated =%s"),
+				*LocalFilePath);
+		}
+		else
+		{
+			UE_LOG(LogHTTPChunkInstaller, Warning, TEXT("WriteUserFile request complete. Local file cache failed to update =%s"),
+				*LocalFilePath);
+		}
+	}
+
+	/**
+	* Should use the initialization constructor instead
+	*/
+	FOnlineTitleFileHttp();
+
+	/** Config based url for enumerating list of cloud files*/
+	FString EnumerateFilesUrl;
+	/** Config based url for accessing the HTTP server */
+	FString BaseUrl;
+
+	FString GetBaseUrl()
+	{
+		return TEXT("http://") + BaseUrl + TEXT("/");
+	}
+
+	/** Info used to send request to authenticate a user *
+	struct FPendingEnumerate
+	{
+		FPendingEnumerate(const FPagedQuery& InPage = FPagedQuery())
+			: Page(InPage)
+		{
+		}
+		FPagedQuery Page;
+	};
+	/** List of pending Http requests for enumerating files */
+	TMap<FHttpRequestPtr, FCloudPagedQuery> EnumerateFilesRequests;
+
+	/** Info used to send request for a file */
+	struct FPendingFileRequest
+	{
+		/**
+		* Constructor
+		*/
+		FPendingFileRequest(const FString& InFileName = FString(TEXT("")))
+			: FileName(InFileName)
+		{
+
+		}
+
+		/**
+		* Equality op
+		*/
+		inline bool operator==(const FPendingFileRequest& Other) const
+		{
+			return FileName == Other.FileName;
+		}
+
+		/** File being operated on by the pending request */
+		FString FileName;
+	};
+	/** List of pending Http requests for reading files */
+	TMap<FHttpRequestPtr, FPendingFileRequest> FileRequests;
+	TMap<FHttpRequestPtr, FPendingFileRequest> FileProgressRequestsMap;
+
+	TArray<FCloudHeader> FileHeaders;
+	TArray<FCloudEntry> Files;
+	bool bCacheFiles;
+	bool bPlatformSupportsSHA256;
+
+	/** Information about local file reads that are in progress */
+	struct FTitleAsyncReadData
+	{
+		/** Name of the file being loaded */
+		FString Filename;
+		/** Amount of data that has been loaded on the async thread so far */
+		FThreadSafeCounter64 BytesRead;
+		/** Bytes read last time game thread noticed */
+		int64 LastBytesRead;
+		/** Async tasks doing the work */
+		FAsyncTask<FTitleFileHttpAsyncLoadAndVerify>*  AsyncTask;
+
+		FTitleAsyncReadData() :
+			LastBytesRead(0),
+			AsyncTask(nullptr)
+		{ }
+
+		bool operator==(const FTitleAsyncReadData& Other) const
+		{
+			return Filename == Other.Filename && AsyncTask == Other.AsyncTask;
+		}
+	};
+	/** Holds the outstanding tasks for hitch free loading and hash calculation */
+	TIndirectArray<FTitleAsyncReadData> AsyncLocalReads;
+};
+
+// Helper class to find all pak file manifests.
+class FChunkSearchVisitor: public IPlatformFile::FDirectoryVisitor
+{
+public:
+	TArray<FString>				PakManifests;
+
+	FChunkSearchVisitor()
+	{}
+
+	virtual bool Visit(const TCHAR* FilenameOrDirectory,bool bIsDirectory)
+	{
+		if(bIsDirectory == false)
+		{
+			FString Filename(FilenameOrDirectory);
+			if(FPaths::GetBaseFilename(Filename).MatchesWildcard("*.manifest"))
+			{
+				PakManifests.AddUnique(Filename);
+			}
+		}
+		return true;
+	}
+};
+
+
+FHTTPChunkInstall::FHTTPChunkInstall()
+	: InstallingChunkID(-1)
+	, InstallerState(ChunkInstallState::Setup)
+	, InstallSpeed(EChunkInstallSpeed::Fast)
+	, bDebugNoInstalledRequired(false)
+	, bFirstRun(true)
+	, bSystemInitialised(false)
+{
+
+}
+
+
+FHTTPChunkInstall::~FHTTPChunkInstall()
+{
+	if (InstallService.IsValid())
+	{
+		InstallService->CancelInstall();
+		InstallService.Reset();
+	}
+}
+
+bool FHTTPChunkInstall::Tick(float DeltaSeconds)
+{
+	if (!bSystemInitialised)
+	{
+		InitialiseSystem();
+	}
+
+	switch (InstallerState)
+	{
+	case ChunkInstallState::Setup:
+		{
+			check(OnlineTitleFile.IsValid());
+			EnumFilesCompleteHandle = OnlineTitleFile->AddOnEnumerateFilesCompleteDelegate_Handle(FOnEnumerateFilesCompleteDelegate::CreateRaw(this,&FHTTPChunkInstall::OSSEnumerateFilesComplete));
+			ReadFileCompleteHandle = OnlineTitleFile->AddOnReadFileCompleteDelegate_Handle(FOnReadFileCompleteDelegate::CreateRaw(this,&FHTTPChunkInstall::OSSReadFileComplete));
+			ChunkSetupTask.SetupWork(BPSModule, InstallDir, ContentDir, HoldingDir, MountedPaks);
+			ChunkSetupTaskThread.Reset(FRunnableThread::Create(&ChunkSetupTask, TEXT("Chunk discovery thread")));
+			InstallerState = ChunkInstallState::SetupWait;
+		} break;
+	case ChunkInstallState::SetupWait:
+		{
+			if (ChunkSetupTask.IsDone())
+			{
+				ChunkSetupTaskThread->WaitForCompletion();
+				ChunkSetupTaskThread.Reset();
+				for (auto It = ChunkSetupTask.InstalledChunks.CreateConstIterator(); It; ++It)
+				{
+					UE_LOG(LogHTTPChunkInstaller, Log, TEXT("Adding Chunk %d to installed manifests"), It.Key());
+					InstalledManifests.Add(It.Key(), It.Value());
+				}
+				for (auto It = ChunkSetupTask.HoldingChunks.CreateConstIterator(); It; ++It)
+				{
+					UE_LOG(LogHTTPChunkInstaller, Log, TEXT("Adding Chunk %d to holding manifests"), It.Key());
+					PrevInstallManifests.Add(It.Key(), It.Value());
+				}
+				MountedPaks.Append(ChunkSetupTask.MountedPaks);
+				InstallerState = ChunkInstallState::QueryRemoteManifests;
+			}
+		} break;
+	case ChunkInstallState::QueryRemoteManifests:
+		{			
+			//Now query the title file service for the chunk manifests. This should return the list of expected chunk manifests
+			check(OnlineTitleFile.IsValid());
+			OnlineTitleFile->ClearFiles();
+			InstallerState = ChunkInstallState::RequestingTitleFiles;
+			UE_LOG(LogHTTPChunkInstaller, Log, TEXT("Enumerating manifest files"));
+			OnlineTitleFile->EnumerateFiles();
+		} break;
+	case ChunkInstallState::SearchTitleFiles:
+		{
+			FString CleanName;
+			TArray<FCloudHeader> FileList;
+			TitleFilesToRead.Reset();
+			RemoteManifests.Reset();
+			ExpectedChunks.Empty();
+			OnlineTitleFile->GetFileList(FileList);
+			for (int32 FileIndex = 0, FileCount = FileList.Num(); FileIndex < FileCount; ++FileIndex)
+			{
+				if (FileList[FileIndex].FileName.MatchesWildcard(TEXT("*.manifest")))
+				{
+					UE_LOG(LogHTTPChunkInstaller, Log, TEXT("Found manifest %s"), *FileList[FileIndex].FileName);
+					TitleFilesToRead.Add(FileList[FileIndex]);
+				}
+			}
+			InstallerState = ChunkInstallState::ReadTitleFiles;
+		} break;
+	case ChunkInstallState::ReadTitleFiles:
+		{
+			if (TitleFilesToRead.Num() > 0 && InstallSpeed != EChunkInstallSpeed::Paused)
+			{
+				if (!IsDataInFileCache(TitleFilesToRead[0].Hash))
+				{
+					UE_LOG(LogHTTPChunkInstaller, Log, TEXT("Reading manifest %s from remote source"), *TitleFilesToRead[0].FileName);
+					InstallerState = ChunkInstallState::WaitingOnRead;
+					OnlineTitleFile->ReadFile(TitleFilesToRead[0].DLName);
+				}
+				else
+				{
+					InstallerState = ChunkInstallState::ReadComplete;
+				}
+			}
+			else
+			{
+				InstallerState = ChunkInstallState::PostSetup;
+			}
+		} break;
+	case ChunkInstallState::ReadComplete:
+		{
+			FileContentBuffer.Reset();
+			bool bReadOK = false;
+			bool bAlreadyLoaded = ManifestsInMemory.Contains(TitleFilesToRead[0].Hash);
+			if (!IsDataInFileCache(TitleFilesToRead[0].Hash))
+			{
+				bReadOK = OnlineTitleFile->GetFileContents(TitleFilesToRead[0].DLName, FileContentBuffer);
+				if (bReadOK)
+				{
+					AddDataToFileCache(TitleFilesToRead[0].Hash, FileContentBuffer);
+				}
+			}
+			else if (!bAlreadyLoaded)
+			{
+				bReadOK = GetDataFromFileCache(TitleFilesToRead[0].Hash, FileContentBuffer);
+				if (!bReadOK)
+				{
+					RemoveDataFromFileCache(TitleFilesToRead[0].Hash);
+				}
+			}
+			if (bReadOK)
+			{
+				if (!bAlreadyLoaded)
+				{
+					ParseTitleFileManifest(TitleFilesToRead[0].Hash);
+				}
+				// Even if the Parse failed remove the file from the list
+				TitleFilesToRead.RemoveAt(0);
+			}
+			if (TitleFilesToRead.Num() == 0)
+			{
+				if (bFirstRun)
+				{
+					ChunkMountTask.SetupWork(BPSModule, ContentDir, MountedPaks, ExpectedChunks);
+					ChunkMountTaskThread.Reset(FRunnableThread::Create(&ChunkMountTask, TEXT("Chunk mounting thread")));
+				}
+				InstallerState = ChunkInstallState::PostSetup;
+			}
+			else
+			{
+				InstallerState = ChunkInstallState::ReadTitleFiles;
+			}
+		} break;
+	case ChunkInstallState::EnterOfflineMode:
+		{
+			for (auto It = InstalledManifests.CreateConstIterator(); It; ++It)
+			{
+				ExpectedChunks.Add(It.Key());
+			}
+			ChunkMountTask.SetupWork(BPSModule, ContentDir, MountedPaks, ExpectedChunks);
+			ChunkMountTaskThread.Reset(FRunnableThread::Create(&ChunkMountTask, TEXT("Chunk mounting thread")));
+			InstallerState = ChunkInstallState::PostSetup;
+		} break;
+	case ChunkInstallState::PostSetup:
+		{
+			if (bFirstRun)
+			{
+				if (ChunkMountTask.IsDone())
+				{
+					ChunkMountTaskThread->WaitForCompletion();
+					ChunkMountTaskThread.Reset();
+					MountedPaks.Append(ChunkMountTask.MountedPaks);
+					UE_LOG(LogHTTPChunkInstaller, Log, TEXT("Completed First Run"));
+					bFirstRun = false;
+					if (PriorityQueue.Num() == 0)
+					{
+						SetInstallSpeed(EChunkInstallSpeed::Paused);
+					}
+				}
+			}
+			else
+			{
+				InstallerState = ChunkInstallState::Idle;
+			}
+		} break;
+	case ChunkInstallState::Idle: 
+		{
+			UpdatePendingInstallQueue();
+		} break;
+	case ChunkInstallState::CopyToContent:
+		{
+			if (!ChunkCopyInstall.IsDone() || !InstallService->IsComplete())
+			{
+				break;
+			}
+			check(InstallingChunkID != -1);
+			if (InstallService.IsValid())
+			{
+				InstallService.Reset();
+			}
+			ChunkCopyInstallThread.Reset();
+			check(RemoteManifests.Find(InstallingChunkID));
+			UE_LOG(LogHTTPChunkInstaller, Log, TEXT("Adding Chunk %d to installed manifests"), InstallingChunkID);
+			InstalledManifests.Add(InstallingChunkID, InstallingChunkManifest);
+			UE_LOG(LogHTTPChunkInstaller, Log, TEXT("Removing Chunk %d from remote manifests"), InstallingChunkID);
+			RemoteManifests.Remove(InstallingChunkID, InstallingChunkManifest);
+			MountedPaks.Append(ChunkCopyInstall.MountedPaks);
+			if (!RemoteManifests.Contains(InstallingChunkID))
+			{
+				// No more manifests relating to the chunk ID are left to install.
+				// Inform any listeners that the install has been completed.
+				FPlatformChunkInstallCompleteMultiDelegate* FoundDelegate = DelegateMap.Find(InstallingChunkID);
+				if (FoundDelegate)
+				{
+					FoundDelegate->Broadcast(InstallingChunkID);
+				}
+			}
+			EndInstall();
+
+		} break;
+	case ChunkInstallState::Installing:
+	case ChunkInstallState::RequestingTitleFiles:
+	case ChunkInstallState::WaitingOnRead:
+	default:
+		break;
+	}
+
+	if (OnlineTitleFileHttp.IsValid())
+	{
+		static_cast<FOnlineTitleFileHttp*>(OnlineTitleFileHttp.Get())->Tick(DeltaSeconds);
+	}
+
+	return true;
+}
+
+void FHTTPChunkInstall::UpdatePendingInstallQueue()
+{
+	if (InstallingChunkID != -1
+#if !UE_BUILD_SHIPPING
+		|| bDebugNoInstalledRequired
+#endif
+	)
+	{
+		return;
+	}
+
+	check(!InstallService.IsValid());
+	bool bPatch = false;
+	while (PriorityQueue.Num() > 0 && InstallerState != ChunkInstallState::Installing)
+	{
+		const FChunkPrio& NextChunk = PriorityQueue[0];
+		TArray<IBuildManifestPtr> FoundChunkManifests;
+		RemoteManifests.MultiFind(NextChunk.ChunkID, FoundChunkManifests);
+		if (FoundChunkManifests.Num() > 0)
+		{
+			auto ChunkManifest = FoundChunkManifests[0];
+			auto ChunkIDField = ChunkManifest->GetCustomField("ChunkID");
+			if (ChunkIDField.IsValid())
+			{
+				BeginChunkInstall(NextChunk.ChunkID, ChunkManifest, FindPreviousInstallManifest(ChunkManifest));
+			}
+			else
+			{
+				PriorityQueue.RemoveAt(0);
+			}
+		}
+		else
+		{
+			PriorityQueue.RemoveAt(0);
+		}
+	}
+	if (InstallingChunkID == -1)
+	{
+		// Install the first available chunk
+		for (auto It = RemoteManifests.CreateConstIterator(); It; ++It)
+		{
+			if (It)
+			{
+				IBuildManifestPtr ChunkManifest = It.Value();
+				auto ChunkIDField = ChunkManifest->GetCustomField("ChunkID");
+				if (ChunkIDField.IsValid())
+				{
+					BeginChunkInstall(ChunkIDField->AsInteger(), ChunkManifest, FindPreviousInstallManifest(ChunkManifest));
+					return;
+				}
+			}
+		}
+	}
+}
+
+
+EChunkLocation::Type FHTTPChunkInstall::GetChunkLocation(uint32 ChunkID)
+{
+#if !UE_BUILD_SHIPPING
+	if(bDebugNoInstalledRequired)
+	{
+		return EChunkLocation::BestLocation;
+	}
+#endif
+
+	// Safe to assume Chunk0 is ready
+	if (ChunkID == 0)
+	{
+		return EChunkLocation::BestLocation;
+	}
+
+	if (bFirstRun || !bSystemInitialised)
+	{
+		/** Still waiting on setup to finish, report that nothing is installed yet... */
+		return EChunkLocation::NotAvailable;
+	}
+	TArray<IBuildManifestPtr> FoundManifests;
+	RemoteManifests.MultiFind(ChunkID, FoundManifests);
+	if (FoundManifests.Num() > 0)
+	{
+		return EChunkLocation::NotAvailable;
+	}
+
+	InstalledManifests.MultiFind(ChunkID, FoundManifests);
+	if (FoundManifests.Num() > 0)
+	{
+		return EChunkLocation::BestLocation;
+	}
+
+	return EChunkLocation::DoesNotExist;
+}
+
+
+float FHTTPChunkInstall::GetChunkProgress(uint32 ChunkID,EChunkProgressReportingType::Type ReportType)
+{
+#if !UE_BUILD_SHIPPING
+	if (bDebugNoInstalledRequired)
+	{
+		return 100.f;
+	}
+#endif
+
+	// Safe to assume Chunk0 is ready
+	if (ChunkID == 0)
+	{
+		return 100.f;
+	}
+
+	if (bFirstRun || !bSystemInitialised)
+	{
+		/** Still waiting on setup to finish, report that nothing is installed yet... */
+		return 0.f;
+	}
+	TArray<IBuildManifestPtr> FoundManifests;
+	RemoteManifests.MultiFind(ChunkID, FoundManifests);
+	if (FoundManifests.Num() > 0)
+	{
+		float Progress = 0;
+		if (InstallingChunkID == ChunkID && InstallService.IsValid())
+		{
+			Progress = InstallService->GetUpdateProgress();
+		}
+		return Progress / FoundManifests.Num();
+	}
+
+	InstalledManifests.MultiFind(ChunkID, FoundManifests);
+	if (FoundManifests.Num() > 0)
+	{
+		return 100.f;
+	}
+
+	return 0.f;
+}
+
+void FHTTPChunkInstall::OSSEnumerateFilesComplete(bool bSuccess)
+{
+	InstallerState = bSuccess ? ChunkInstallState::SearchTitleFiles : ChunkInstallState::EnterOfflineMode;
+}
+
+void FHTTPChunkInstall::OSSReadFileComplete(bool bSuccess, const FString& Filename)
+{
+	InstallerState = bSuccess ? ChunkInstallState::ReadComplete : ChunkInstallState::EnterOfflineMode;
+}
+
+void FHTTPChunkInstall::OSSInstallComplete(bool bSuccess, IBuildManifestRef BuildManifest)
+{
+	if (bSuccess)
+	{
+		// Completed OK. Write the manifest. If the chunk doesn't exist, copy to the content dir.
+		// Otherwise, writing the manifest will prompt a copy on next start of the game
+		FString ManifestName;
+		FString ChunkFdrName;
+		uint32 ChunkID;
+		bool bIsPatch;
+		if (!BuildChunkFolderName(BuildManifest, ChunkFdrName, ManifestName, ChunkID, bIsPatch))
+		{
+			//Something bad has happened, bail
+			EndInstall();
+			return;
+		}
+		UE_LOG(LogHTTPChunkInstaller, Log, TEXT("Chunk %d install complete, preparing to copy to content directory"), ChunkID);
+		FString ManifestPath = FPaths::Combine(*InstallDir, *ChunkFdrName, *ManifestName);
+		FString HoldingManifestPath = FPaths::Combine(*HoldingDir, *ChunkFdrName, *ManifestName);
+		FString SrcDir = FPaths::Combine(*InstallDir, *ChunkFdrName);
+		FString DestDir = FPaths::Combine(*ContentDir, *ChunkFdrName);
+		bool bCopyDir = InstallDir != ContentDir; 
+		TArray<IBuildManifestPtr> FoundManifests;
+		InstalledManifests.MultiFind(ChunkID, FoundManifests);
+		for (const auto& It : FoundManifests)
+		{
+			auto FoundPatchField = It->GetCustomField("bIsPatch");
+			bool bFoundPatch = FoundPatchField.IsValid() ? FoundPatchField->AsString() == TEXT("true") : false;
+			if (bFoundPatch == bIsPatch)
+			{
+				bCopyDir = false;
+			}
+		}
+		ChunkCopyInstall.SetupWork(ManifestPath, HoldingManifestPath, SrcDir, DestDir, BPSModule, BuildManifest, MountedPaks, bCopyDir);
+		UE_LOG(LogHTTPChunkInstaller, Log, TEXT("Copying Chunk %d to content directory"), ChunkID);
+		ChunkCopyInstallThread.Reset(FRunnableThread::Create(&ChunkCopyInstall, TEXT("Chunk Install Copy Thread")));
+		InstallerState = ChunkInstallState::CopyToContent;
+	}
+	else
+	{
+		//Something bad has happened, return to the Idle state. We'll re-attempt the install
+		EndInstall();
+	}
+}
+
+void FHTTPChunkInstall::ParseTitleFileManifest(const FString& ManifestFileHash)
+{
+#if !UE_BUILD_SHIPPING
+	if (bDebugNoInstalledRequired)
+	{
+		// Forces the installer to think that no remote manifests exist, so nothing needs to be installed.
+		return;
+	}
+#endif
+	FString JsonBuffer;
+	FFileHelper::BufferToString(JsonBuffer, FileContentBuffer.GetData(), FileContentBuffer.Num());
+	auto RemoteManifest = BPSModule->MakeManifestFromJSON(JsonBuffer);
+	if (!RemoteManifest.IsValid())
+	{
+		UE_LOG(LogHTTPChunkInstaller, Warning, TEXT("Manifest was invalid"));
+		return;
+	}
+	auto RemoteChunkIDField = RemoteManifest->GetCustomField("ChunkID");
+	if (!RemoteChunkIDField.IsValid())
+	{
+		UE_LOG(LogHTTPChunkInstaller, Warning, TEXT("Manifest ChunkID was invalid or missing"));
+		return;
+	}
+	//Compare to installed manifests and add to the remote if it needs to be installed.
+	uint32 ChunkID = (uint32)RemoteChunkIDField->AsInteger();
+	ExpectedChunks.Add(ChunkID);
+	TArray<IBuildManifestPtr> FoundManifests;
+	InstalledManifests.MultiFind(ChunkID, FoundManifests);
+	uint32 FoundCount = FoundManifests.Num();
+	if (FoundCount > 0)
+	{
+		auto RemotePatchManifest = RemoteManifest->GetCustomField("bIsPatch");
+		auto RemoteVersion = RemoteManifest->GetVersionString();
+		bool bRemoteIsPatch = RemotePatchManifest.IsValid() ? RemotePatchManifest->AsString() == TEXT("true") : false;
+		for (uint32 FoundIndex = 0; FoundIndex < FoundCount; ++FoundIndex)
+		{
+			const auto& InstalledManifest = FoundManifests[FoundIndex];
+			auto InstalledVersion = InstalledManifest->GetVersionString();
+			auto InstallPatchManifest = InstalledManifest->GetCustomField("bIsPatch");
+			bool bInstallIsPatch = InstallPatchManifest.IsValid() ? InstallPatchManifest->AsString() == TEXT("true") : false;
+			if (InstalledVersion != RemoteVersion && bInstallIsPatch == bRemoteIsPatch)
+			{
+				UE_LOG(LogHTTPChunkInstaller, Log, TEXT("Adding Chunk %d to remote manifests"), ChunkID);
+				RemoteManifests.Add(ChunkID, RemoteManifest);
+				if(!ManifestFileHash.IsEmpty())
+				{
+					ManifestsInMemory.Add(ManifestFileHash);
+				}
+				//Remove from the installed map
+				if (bFirstRun)
+				{
+					// Prevent the paks from being mounted by removing the manifest file
+					FString ChunkFdrName;
+					FString ManifestName;
+					bool bIsPatch;
+					if (BuildChunkFolderName(InstalledManifest.ToSharedRef(), ChunkFdrName, ManifestName, ChunkID, bIsPatch))
+					{
+						FString ManifestPath = FPaths::Combine(*ContentDir, *ChunkFdrName, *ManifestName);
+						FString HoldingPath = FPaths::Combine(*HoldingDir, *ChunkFdrName, *ManifestName);
+						IPlatformFile& PlatformFile = FPlatformFileManager::Get().GetPlatformFile();
+						PlatformFile.CreateDirectoryTree(*FPaths::Combine(*HoldingDir, *ChunkFdrName));
+						PlatformFile.MoveFile(*HoldingPath, *ManifestPath);
+					}
+					UE_LOG(LogHTTPChunkInstaller, Log, TEXT("Adding Chunk %d to previous installed manifests"), ChunkID);
+					PrevInstallManifests.Add(ChunkID, InstalledManifest);
+					UE_LOG(LogHTTPChunkInstaller, Log, TEXT("Removing Chunk %d from installed manifests"), ChunkID);
+					InstalledManifests.Remove(ChunkID, InstalledManifest);
+				}
+			}
+		}
+	}
+	else
+	{
+		UE_LOG(LogHTTPChunkInstaller, Log, TEXT("Adding Chunk %d to remote manifests"), ChunkID);
+		RemoteManifests.Add(ChunkID, RemoteManifest);
+		if (!ManifestFileHash.IsEmpty())
+		{
+			ManifestsInMemory.Add(ManifestFileHash);
+		}
+	}
+}
+
+bool FHTTPChunkInstall::BuildChunkFolderName(IBuildManifestRef Manifest, FString& ChunkFdrName, FString& ManifestName, uint32& ChunkID, bool& bIsPatch)
+{
+	auto ChunkIDField = Manifest->GetCustomField("ChunkID");
+	auto ChunkPatchField = Manifest->GetCustomField("bIsPatch");
+
+	if (!ChunkIDField.IsValid())
+	{
+		return false;
+	}
+	ChunkID = ChunkIDField->AsInteger();
+	bIsPatch = ChunkPatchField.IsValid() ? ChunkPatchField->AsString() == TEXT("true") : false;
+	ManifestName = FString::Printf(TEXT("chunk_%u"), ChunkID);
+	if (bIsPatch)
+	{
+		ManifestName += TEXT("_patch");
+	}
+	ManifestName += TEXT(".manifest");
+	ChunkFdrName = FString::Printf(TEXT("%s%d"), !bIsPatch ? TEXT("base") : TEXT("patch"), ChunkID);
+	return true;
+}
+
+bool FHTTPChunkInstall::PrioritizeChunk(uint32 ChunkID, EChunkPriority::Type Priority)
+{	
+	int32 FoundIndex;
+	PriorityQueue.Find(FChunkPrio(ChunkID, Priority), FoundIndex);
+	if (FoundIndex != INDEX_NONE)
+	{
+		PriorityQueue.RemoveAt(FoundIndex);
+	}
+	// Low priority is assumed if the chunk ID doesn't exist in the queue
+	if (Priority != EChunkPriority::Low)
+	{
+		PriorityQueue.AddUnique(FChunkPrio(ChunkID, Priority));
+		PriorityQueue.Sort();
+	}
+	return true;
+}
+
+FDelegateHandle FHTTPChunkInstall::SetChunkInstallDelgate(uint32 ChunkID, FPlatformChunkInstallCompleteDelegate Delegate)
+{
+	FPlatformChunkInstallCompleteMultiDelegate* FoundDelegate = DelegateMap.Find(ChunkID);
+	if (FoundDelegate)
+	{
+		return FoundDelegate->Add(Delegate);
+	}
+	else
+	{
+		FPlatformChunkInstallCompleteMultiDelegate MC;
+		auto RetVal = MC.Add(Delegate);
+		DelegateMap.Add(ChunkID, MC);
+		return RetVal;
+	}
+	return FDelegateHandle();
+}
+
+void FHTTPChunkInstall::RemoveChunkInstallDelgate(uint32 ChunkID, FDelegateHandle Delegate)
+{
+	FPlatformChunkInstallCompleteMultiDelegate* FoundDelegate = DelegateMap.Find(ChunkID);
+	if (!FoundDelegate)
+	{
+		return;
+	}
+	FoundDelegate->Remove(Delegate);
+}
+
+void FHTTPChunkInstall::BeginChunkInstall(uint32 ChunkID,IBuildManifestPtr ChunkManifest, IBuildManifestPtr PrevInstallChunkManifest)
+{
+	check(ChunkManifest->GetCustomField("ChunkID").IsValid());
+	InstallingChunkID = ChunkID;
+	check(ChunkID > 0);
+	InstallingChunkManifest = ChunkManifest;
+	IPlatformFile& PlatformFile = FPlatformFileManager::Get().GetPlatformFile();
+	auto PatchField = ChunkManifest->GetCustomField("bIsPatch");
+	bool bIsPatch = PatchField.IsValid() ? PatchField->AsString() == TEXT("true") : false;
+	auto ChunkFolderName = FString::Printf(TEXT("%s%d"),!bIsPatch ? TEXT("base") : TEXT("patch"), InstallingChunkID);
+	auto ChunkInstallDir = FPaths::Combine(*InstallDir,*ChunkFolderName);
+	auto ChunkStageDir = FPaths::Combine(*StageDir,*ChunkFolderName);
+	if(!PlatformFile.DirectoryExists(*ChunkStageDir))
+	{
+		PlatformFile.CreateDirectoryTree(*ChunkStageDir);
+	}
+	if(!PlatformFile.DirectoryExists(*ChunkInstallDir))
+	{
+		PlatformFile.CreateDirectoryTree(*ChunkInstallDir);
+	}
+	BPSModule->SetCloudDirectory(CloudDir + TEXT("/") + CloudDirectory);
+	BPSModule->SetStagingDirectory(ChunkStageDir);
+	UE_LOG(LogHTTPChunkInstaller,Log,TEXT("Starting Chunk %d install"),InstallingChunkID);
+	InstallService = BPSModule->StartBuildInstall(PrevInstallChunkManifest,ChunkManifest,ChunkInstallDir,FBuildPatchBoolManifestDelegate::CreateRaw(this,&FHTTPChunkInstall::OSSInstallComplete));
+	if(InstallSpeed == EChunkInstallSpeed::Paused && !InstallService->IsPaused())
+	{
+		InstallService->TogglePauseInstall();
+	}
+	InstallerState = ChunkInstallState::Installing;
+}
+
+/**
+ * Note: the following cache functions are synchronous and may need to become asynchronous...
+ */
+bool FHTTPChunkInstall::AddDataToFileCache(const FString& ManifestHash,const TArray<uint8>& Data)
+{
+	if (ManifestHash.IsEmpty())
+	{
+		return false;
+	}
+	UE_LOG(LogHTTPChunkInstaller, Log, TEXT("Adding data hash %s to file cache"), *ManifestHash);
+	return FFileHelper::SaveArrayToFile(Data, *FPaths::Combine(*CacheDir, *ManifestHash));
+}
+
+bool FHTTPChunkInstall::IsDataInFileCache(const FString& ManifestHash)
+{
+	if(ManifestHash.IsEmpty())
+	{
+		return false;
+	}
+	IPlatformFile& PlatformFile = FPlatformFileManager::Get().GetPlatformFile();
+	return PlatformFile.FileExists(*FPaths::Combine(*CacheDir, *ManifestHash));
+}
+
+bool FHTTPChunkInstall::GetDataFromFileCache(const FString& ManifestHash, TArray<uint8>& Data)
+{
+	if(ManifestHash.IsEmpty())
+	{
+		return false;
+	}
+	UE_LOG(LogHTTPChunkInstaller, Log, TEXT("Reading data hash %s from file cache"), *ManifestHash);
+	return FFileHelper::LoadFileToArray(Data, *FPaths::Combine(*CacheDir, *ManifestHash));
+}
+
+bool FHTTPChunkInstall::RemoveDataFromFileCache(const FString& ManifestHash)
+{
+	if(ManifestHash.IsEmpty())
+	{
+		return false;
+	}
+	UE_LOG(LogHTTPChunkInstaller, Log, TEXT("Removing data hash %s from file cache"), *ManifestHash);
+	IPlatformFile& PlatformFile = FPlatformFileManager::Get().GetPlatformFile();
+	auto ManifestPath = FPaths::Combine(*CacheDir, *ManifestHash);
+	if (PlatformFile.FileExists(*ManifestPath))
+	{
+		return PlatformFile.DeleteFile(*ManifestPath);
+	}
+	return false;
+}
+
+void FHTTPChunkInstall::InitialiseSystem()
+{
+	BPSModule = GetBuildPatchServices();
+
+#if !UE_BUILD_SHIPPING
+	const TCHAR* CmdLine = FCommandLine::Get();
+	if (!FPlatformProperties::RequiresCookedData() || FParse::Param(CmdLine, TEXT("NoPak")) || FParse::Param(CmdLine, TEXT("NoChunkInstall")))
+	{
+		bDebugNoInstalledRequired = true;
+	}
+#endif
+
+	// Grab the title file interface
+	FString TitleFileSource;
+	bool bValidTitleFileSource = GConfig->GetString(TEXT("HTTPChunkInstall"), TEXT("TitleFileSource"), TitleFileSource, GEngineIni);
+	if (bValidTitleFileSource && TitleFileSource == TEXT("Http"))
+	{
+		OnlineTitleFile = OnlineTitleFileHttp = MakeShareable(new FOnlineTitleFileHttp(CloudDir));
+	}
+/*	else if (bValidTitleFileSource && TitleFileSource != TEXT("Local"))
+	{
+		OnlineTitleFile = Online::GetTitleFileInterface(*TitleFileSource);
+	}*/
+	else
+	{
+		FString LocalTileFileDirectory = FPaths::GameConfigDir();
+		auto bGetConfigDir = GConfig->GetString(TEXT("HTTPChunkInstall"), TEXT("LocalTitleFileDirectory"), LocalTileFileDirectory, GEngineIni);
+		OnlineTitleFile = MakeShareable(new FLocalTitleFile(LocalTileFileDirectory));
+#if !UE_BUILD_SHIPPING
+		bDebugNoInstalledRequired = !bGetConfigDir;
+#endif
+	}
+	CloudDirectory = TEXT("");
+	CloudDir = FPaths::Combine(*FPaths::GameContentDir(), TEXT("Cloud"));
+	StageDir = FPaths::Combine(*FPaths::GameSavedDir(), TEXT("Chunks"), TEXT("Staged"));
+	InstallDir = FPaths::Combine(*FPaths::GameSavedDir(), TEXT("Chunks"), TEXT("Installed")); // By default this should match ContentDir
+	BackupDir = FPaths::Combine(*FPaths::GameSavedDir(), TEXT("Chunks"), TEXT("Backup"));
+	CacheDir = FPaths::Combine(*FPaths::GameSavedDir(), TEXT("Chunks"), TEXT("Cache"));
+	HoldingDir = FPaths::Combine(*FPaths::GameSavedDir(), TEXT("Chunks"), TEXT("Hold"));
+	ContentDir = FPaths::Combine(*FPaths::GameSavedDir(), TEXT("Chunks"), TEXT("Installed")); // By default this should match InstallDir
+
+	FString TmpString1;
+	FString TmpString2;
+	if (GConfig->GetString(TEXT("HTTPChunkInstall"), TEXT("CloudDirectory"), TmpString1, GEngineIni))
+	{
+		CloudDirectory = CloudDir = TmpString1;
+	}
+	if ((GConfig->GetString(TEXT("HTTPChunkInstall"), TEXT("CloudProtocol"), TmpString1, GEngineIni)) && (GConfig->GetString(TEXT("HTTPChunkInstall"), TEXT("CloudDomain"), TmpString2, GEngineIni)))
+	{
+		CloudDir = FString::Printf(TEXT("%s://%s"), *TmpString1, *TmpString2);
+	}
+	if (GConfig->GetString(TEXT("HTTPChunkInstall"), TEXT("StageDirectory"), TmpString1, GEngineIni))
+	{
+		StageDir = TmpString1;
+	}
+	if (GConfig->GetString(TEXT("HTTPChunkInstall"), TEXT("InstallDirectory"), TmpString1, GEngineIni))
+	{
+		InstallDir = TmpString1;
+	}
+	if (GConfig->GetString(TEXT("HTTPChunkInstall"), TEXT("BackupDirectory"), TmpString1, GEngineIni))
+	{
+		BackupDir = TmpString1;
+	}
+	if (GConfig->GetString(TEXT("HTTPChunkInstall"), TEXT("ContentDirectory"), TmpString1, GEngineIni))
+	{
+		ContentDir = TmpString1;
+	}
+	if (GConfig->GetString(TEXT("HTTPChunkInstall"), TEXT("HoldingDirectory"), TmpString1, GEngineIni))
+	{
+		HoldingDir = TmpString1;
+	}
+
+	bFirstRun = true;
+	bSystemInitialised = true;
+}
+
+IBuildManifestPtr FHTTPChunkInstall::FindPreviousInstallManifest(const IBuildManifestPtr& ChunkManifest)
+{
+	auto ChunkIDField = ChunkManifest->GetCustomField("ChunkID");
+	if (!ChunkIDField.IsValid())
+	{
+		return IBuildManifestPtr();
+	}
+	auto ChunkID = ChunkIDField->AsInteger();
+	TArray<IBuildManifestPtr> FoundManifests;
+	PrevInstallManifests.MultiFind(ChunkID, FoundManifests);
+	return FoundManifests.Num() == 0 ? IBuildManifestPtr() : FoundManifests[0];
+}
+
+void FHTTPChunkInstall::EndInstall()
+{
+	if (InstallService.IsValid())
+	{
+		//InstallService->CancelInstall();
+		InstallService.Reset();
+	}
+	InstallingChunkID = -1;
+	InstallingChunkManifest.Reset();
+	InstallerState = ChunkInstallState::Idle;
+}
+
+/**
+ * Module for the HTTP Chunk Installer
+ */
+class FHTTPChunkInstallerModule : public IPlatformChunkInstallModule
+{
+public:
+	TUniquePtr<IPlatformChunkInstall> ChunkInstaller;
+
+	FHTTPChunkInstallerModule()
+		: ChunkInstaller(new FHTTPChunkInstall())
+	{
+
+	}
+
+	virtual IPlatformChunkInstall* GetPlatformChunkInstall()
+	{
+		return ChunkInstaller.Get();
+	}
+};
+
+IMPLEMENT_MODULE(FHTTPChunkInstallerModule, HTTPChunkInstaller);
\ No newline at end of file
diff --git a/Engine/Plugins/Runtime/HTTPChunkInstaller/Source/Private/HTTPChunkInstallerLog.cpp b/Engine/Plugins/Runtime/HTTPChunkInstaller/Source/Private/HTTPChunkInstallerLog.cpp
new file mode 100644
--- /dev/null
+++ b/Engine/Plugins/Runtime/HTTPChunkInstaller/Source/Private/HTTPChunkInstallerLog.cpp
@@ -0,0 +1,5 @@
+// Copyright 1998-2016 Epic Games, Inc. All Rights Reserved.
+
+#include "HTTPChunkInstallerLog.h"
+
+DEFINE_LOG_CATEGORY(LogHTTPChunkInstaller);
diff --git a/Engine/Plugins/Runtime/HTTPChunkInstaller/Source/Private/LocalTitleFile.cpp b/Engine/Plugins/Runtime/HTTPChunkInstaller/Source/Private/LocalTitleFile.cpp
new file mode 100644
--- /dev/null
+++ b/Engine/Plugins/Runtime/HTTPChunkInstaller/Source/Private/LocalTitleFile.cpp
@@ -0,0 +1,138 @@
+// Copyright 1998-2016 Epic Games, Inc. All Rights Reserved.
+
+#include "LocalTitleFile.h"
+#include "FileManagerGeneric.h"
+#include "Paths.h"
+#include "FileHelper.h"
+
+FLocalTitleFile::FLocalTitleFile(const FString& InRootDirectory)
+: RootDirectory(InRootDirectory)
+{
+	RootDirectory.Replace(TEXT("\\"), TEXT("/"));
+	if (!RootDirectory.EndsWith(TEXT("/")))
+	{
+		RootDirectory.Append(TEXT("/"));
+	}
+}
+
+
+bool FLocalTitleFile::GetFileContents(const FString& DLName, TArray<uint8>& FileContents)
+{
+	const TArray< uint8 >* ExistingFileContents = DLNameToFileContents.Find(DLName);
+
+	if (ExistingFileContents == nullptr)
+	{
+		return false;
+	}
+
+	FileContents.Append(*ExistingFileContents);
+
+	return true;
+}
+
+
+bool FLocalTitleFile::ClearFiles()
+{
+	FileHeaders.Empty();
+	DLNameToFileContents.Empty();
+
+	return true;
+}
+
+
+bool FLocalTitleFile::ClearFile(const FString& DLName)
+{
+	bool ClearedFile = false;
+	const FString FileName = GetFileNameFromDLName(DLName);
+
+	for (int Index = 0; Index < FileHeaders.Num(); Index++)
+	{
+		if (FileHeaders[Index].DLName == DLName)
+		{
+			FileHeaders.RemoveAt(Index);
+			ClearedFile = true;
+		}
+	}
+
+	DLNameToFileContents.Remove(DLName);
+
+	return ClearedFile;
+}
+
+void FLocalTitleFile::DeleteCachedFiles(bool bSkipEnumerated)
+{
+	// not implemented
+}
+
+bool FLocalTitleFile::EnumerateFiles(const FCloudPagedQuery& Page)
+{
+	const FString WildCard = FPaths::Combine(*RootDirectory, TEXT("*"));
+
+	TArray<FString> Filenames;
+	IFileManager::Get().FindFiles(Filenames, *WildCard, true, false);
+
+	for(int32 FileIdx = 0; FileIdx < Filenames.Num(); ++FileIdx)
+	{
+		const FString Filename = Filenames[FileIdx];
+
+		FCloudHeader NewHeader;
+		NewHeader.FileName = Filename;
+		NewHeader.DLName = Filename + FString::Printf(TEXT("%d"),FileIdx);
+		NewHeader.FileSize = 0;
+		NewHeader.Hash.Empty();
+
+		FileHeaders.Add(NewHeader);
+	}
+
+	TriggerOnEnumerateFilesCompleteDelegates(true);
+
+	return true;
+}
+
+void FLocalTitleFile::GetFileList(TArray<FCloudHeader>& InFileHeaders)
+{
+	InFileHeaders.Append(FileHeaders);
+}
+
+
+bool FLocalTitleFile::ReadFile(const FString& DLName)
+{
+	const TArray< uint8 >* ExistingFileContents = DLNameToFileContents.Find(DLName);
+
+	if (ExistingFileContents != nullptr)
+	{
+		TriggerOnReadFileCompleteDelegates(true, DLName);
+
+		return true;
+	}
+
+	const FString FileName = GetFileNameFromDLName(DLName);
+
+	TArray<uint8> FileContents;
+
+	if (!FFileHelper::LoadFileToArray(FileContents, *(RootDirectory + FileName)))
+	{
+		TriggerOnReadFileCompleteDelegates(false, DLName);
+
+		return false;
+	}
+
+	DLNameToFileContents.Add(DLName, FileContents);
+	TriggerOnReadFileCompleteDelegates(true, DLName);
+
+	return true;
+}
+
+
+FString FLocalTitleFile::GetFileNameFromDLName(const FString& DLName) const
+{
+	for (int Index = 0; Index < FileHeaders.Num(); Index++)
+	{
+		if (FileHeaders[Index].DLName == DLName)
+		{
+			return FileHeaders[Index].FileName;
+		}
+	}
+
+	return FString();
+}
\ No newline at end of file
diff --git a/Engine/Plugins/Runtime/HTTPChunkInstaller/Source/Public/ChunkInstall.h b/Engine/Plugins/Runtime/HTTPChunkInstaller/Source/Public/ChunkInstall.h
new file mode 100644
--- /dev/null
+++ b/Engine/Plugins/Runtime/HTTPChunkInstaller/Source/Public/ChunkInstall.h
@@ -0,0 +1,161 @@
+// Copyright 1998-2016 Epic Games, Inc. All Rights Reserved.
+
+#pragma once
+
+#include "CoreMinimal.h"
+#include "Runnable.h"
+#include "Interfaces/IBuildManifest.h"
+#include "IBuildPatchServicesModule.h"
+#include "PlatformProcess.h"
+#include "Event.h"
+#include "GenericPlatformFile.h"
+#include "Paths.h"
+#include "PlatformFileManager.h"
+#include "CoreDelegates.h"
+#include "FileManager.h"
+#include "Stats.h"
+
+class FChunkInstallTask : public FRunnable
+{
+public:
+	/** Input parameters */
+	FString							ManifestPath;
+	FString							HoldingManifestPath;
+	FString							SrcDir;
+	FString							DestDir;
+	IBuildPatchServicesModule*		BPSModule;
+	IBuildManifestPtr				BuildManifest;
+	bool							bCopy;
+	const TArray<FString>*			CurrentMountPaks;
+	FEvent*							CompleteEvent;				
+	/** Output */
+	TArray<FString>					MountedPaks;
+
+	FChunkInstallTask()
+	{
+		CompleteEvent = FPlatformProcess::GetSynchEventFromPool(true);
+	}
+
+	~FChunkInstallTask()
+	{
+		FPlatformProcess::ReturnSynchEventToPool(CompleteEvent);
+		CompleteEvent = nullptr;
+	}
+
+	void SetupWork(FString InManifestPath, FString InHoldingManifestPath, FString InSrcDir, FString InDestDir, IBuildPatchServicesModule* InBPSModule, IBuildManifestRef InBuildManifest, const TArray<FString>& InCurrentMountedPaks, bool bInCopy)
+	{
+		ManifestPath = InManifestPath;
+		HoldingManifestPath = InHoldingManifestPath;
+		SrcDir = InSrcDir;
+		DestDir = InDestDir;
+		BPSModule = InBPSModule;
+		BuildManifest = InBuildManifest;
+		CurrentMountPaks = &InCurrentMountedPaks;
+		bCopy = bInCopy;
+
+		MountedPaks.Reset();
+		CompleteEvent->Reset();
+	}
+
+	void DoWork()
+	{
+		// Helper class to find all pak files.
+		class FPakSearchVisitor : public IPlatformFile::FDirectoryVisitor
+		{
+			TArray<FString>& FoundPakFiles;
+		public:
+			FPakSearchVisitor(TArray<FString>& InFoundPakFiles)
+				: FoundPakFiles(InFoundPakFiles)
+			{}
+			virtual bool Visit(const TCHAR* FilenameOrDirectory, bool bIsDirectory)
+			{
+				if (bIsDirectory == false)
+				{
+					FString Filename(FilenameOrDirectory);
+					if (FPaths::GetExtension(Filename) == TEXT("pak"))
+					{
+						FoundPakFiles.Add(Filename);
+					}
+				}
+				return true;
+			}
+		};
+
+		check(CurrentMountPaks);
+
+		TArray<FString> PakFiles;
+		IPlatformFile& PlatformFile = FPlatformFileManager::Get().GetPlatformFile();
+		BPSModule->SaveManifestToFile(ManifestPath, BuildManifest.ToSharedRef(), false);
+		if (PlatformFile.FileExists(*HoldingManifestPath))
+		{
+			PlatformFile.DeleteFile(*HoldingManifestPath);
+		}
+		if (bCopy)
+		{
+			if (PlatformFile.DirectoryExists(*DestDir))
+			{
+				PlatformFile.DeleteDirectoryRecursively(*DestDir);
+			}
+			PlatformFile.CreateDirectoryTree(*DestDir);
+			if (PlatformFile.CopyDirectoryTree(*DestDir, *SrcDir, true))
+			{
+				PlatformFile.DeleteDirectoryRecursively(*SrcDir);
+			}
+		}
+		// Find all pak files.
+		FPakSearchVisitor Visitor(PakFiles);
+		PlatformFile.IterateDirectoryRecursively(*DestDir, Visitor);
+		auto PakReadOrderField = BuildManifest->GetCustomField("PakReadOrdering");
+		uint32 PakReadOrder = PakReadOrderField.IsValid() ? (uint32)PakReadOrderField->AsInteger() : 0;
+		for (uint32 PakIndex = 0, PakCount = PakFiles.Num(); PakIndex < PakCount; ++PakIndex)
+		{
+			if (!CurrentMountPaks->Contains(PakFiles[PakIndex]) && !MountedPaks.Contains(PakFiles[PakIndex]))
+			{
+				// TODO: are we a patch?
+// 				if (PakFiles[PakIndex].EndsWith("_P.pak"))
+// 				{
+// 					// bump the read prioritiy 
+// 					++PakReadOrder;
+// 				}
+				if (FCoreDelegates::OnMountPak.IsBound())
+				{
+					auto bSuccess = FCoreDelegates::OnMountPak.Execute(PakFiles[PakIndex], PakReadOrder, nullptr);
+#if !UE_BUILD_SHIPPING
+					if (!bSuccess)
+					{
+						// This can fail because of the sandbox system - which the pak system doesn't understand.
+						auto SandboxedPath = IFileManager::Get().ConvertToAbsolutePathForExternalAppForRead(*PakFiles[PakIndex]);
+						bSuccess = FCoreDelegates::OnMountPak.Execute(SandboxedPath, PakReadOrder, nullptr);
+					}
+#endif
+					MountedPaks.Add(PakFiles[PakIndex]);
+				}
+			}
+		}
+		//Register the install
+		BPSModule->RegisterAppInstallation(BuildManifest.ToSharedRef(), DestDir);
+
+		CompleteEvent->Trigger();
+	}
+
+	uint32 Run()
+	{
+		DoWork();
+		return 0;
+	}
+
+	bool IsDone()
+	{
+		return CompleteEvent->Wait(FTimespan(0));
+	}
+
+	static const TCHAR *Name()
+	{
+		return TEXT("FChunkDescovery");
+	}
+
+	FORCEINLINE TStatId GetStatId() const
+	{
+		RETURN_QUICK_DECLARE_CYCLE_STAT(FChunkInstallTask, STATGROUP_ThreadPoolAsyncTasks);
+	}
+};
diff --git a/Engine/Plugins/Runtime/HTTPChunkInstaller/Source/Public/ChunkSetup.h b/Engine/Plugins/Runtime/HTTPChunkInstaller/Source/Public/ChunkSetup.h
new file mode 100644
--- /dev/null
+++ b/Engine/Plugins/Runtime/HTTPChunkInstaller/Source/Public/ChunkSetup.h
@@ -0,0 +1,339 @@
+// Copyright 1998-2016 Epic Games, Inc. All Rights Reserved.
+
+#pragma once
+
+// Helper class to find all pak/mainfests files.
+class FFileSearchVisitor : public IPlatformFile::FDirectoryVisitor
+{
+	FString			 FileWildcard;
+	TArray<FString>& FoundFiles;
+public:
+	FFileSearchVisitor(FString InFileWildcard, TArray<FString>& InFoundFiles)
+		: FileWildcard(InFileWildcard)
+		, FoundFiles(InFoundFiles)
+	{}
+	virtual bool Visit(const TCHAR* FilenameOrDirectory, bool bIsDirectory)
+	{
+		if (bIsDirectory == false)
+		{
+			FString Filename(FilenameOrDirectory);
+			if (Filename.MatchesWildcard(FileWildcard))
+			{
+				FoundFiles.Add(Filename);
+			}
+		}
+		return true;
+	}
+};
+
+class IBuildPatchServicesModule;
+
+class FChunkSetupTask : public FRunnable, public IPlatformFile::FDirectoryVisitor
+{
+public:
+	/** Input parameters */
+	IBuildPatchServicesModule*	BPSModule;
+	FString						InstallDir; // Intermediate directory where installed chunks may be waiting
+	FString						ContentDir; // Directory where installed chunks need to live to be mounted
+	FString						HoldingDir; // Directory where manifest for chunks that are out of date and can be use for updates but not mounted.
+	const TArray<FString>*		CurrentMountPaks; // Array of already mounted Paks
+	/** Output */
+	FEvent*									CompleteEvent;
+	TArray<FString>							MountedPaks;
+	TMultiMap<uint32, IBuildManifestPtr>	InstalledChunks;
+	TMultiMap<uint32, IBuildManifestPtr>	HoldingChunks;
+	/** Working */
+	TArray<FString>		FoundPaks;
+	TArray<FString>		FoundManifests;
+	FFileSearchVisitor	PakVisitor;
+	FFileSearchVisitor	ManifestVisitor;
+	TArray<FString>		ManifestsToRemove;
+	uint32				Pass;
+
+	FChunkSetupTask()
+		: BPSModule(nullptr)
+		, CurrentMountPaks(nullptr)
+		, PakVisitor(TEXT("*.pak"), FoundPaks)
+		, ManifestVisitor(TEXT("*.manifest"), FoundManifests)
+	{
+		CompleteEvent = FPlatformProcess::GetSynchEventFromPool(true);
+	}
+
+	virtual ~FChunkSetupTask()
+	{
+		FPlatformProcess::ReturnSynchEventToPool(CompleteEvent);
+	}
+
+	void SetupWork(IBuildPatchServicesModule* InBPSModule, FString InInstallDir, FString InContentDir, FString InHoldingDir, const TArray<FString>& InCurrentMountedPaks)
+	{
+		BPSModule = InBPSModule;
+		InstallDir = InInstallDir;
+		ContentDir = InContentDir;
+		HoldingDir = InHoldingDir;
+		CurrentMountPaks = &InCurrentMountedPaks;
+
+		Pass = 0;
+		InstalledChunks.Reset();
+		MountedPaks.Reset();
+		FoundManifests.Reset();
+		FoundPaks.Reset();
+		ManifestsToRemove.Reset();
+
+		CompleteEvent->Reset();
+	}
+
+	void DoWork()
+	{
+		IPlatformFile& PlatformFile = FPlatformFileManager::Get().GetPlatformFile();
+		PlatformFile.IterateDirectory(*InstallDir, *this);
+		for (const auto& ToRemove : ManifestsToRemove)
+		{
+			PlatformFile.DeleteFile(*ToRemove);
+		}
+		++Pass;
+		PlatformFile.IterateDirectory(*ContentDir, *this);
+		++Pass;
+		PlatformFile.IterateDirectory(*HoldingDir, *this);
+		CompleteEvent->Trigger();
+	}
+
+	bool Visit(const TCHAR* FilenameOrDirectory, bool bIsDirectory) override
+	{
+		if (!bIsDirectory)
+		{
+			return true;
+		}
+		IPlatformFile& PlatformFile = FPlatformFileManager::Get().GetPlatformFile();
+
+		FoundManifests.Reset();
+		PlatformFile.IterateDirectory(FilenameOrDirectory, ManifestVisitor);
+		if (FoundManifests.Num() == 0)
+		{
+			return true;
+		}
+		
+		// Only allow one manifest per folder, any more suggests corruption so mark the folder for delete
+		if (FoundManifests.Num() > 1)
+		{
+			ManifestsToRemove.Append(FoundManifests);
+			return true;
+		}
+		// Load the manifest, so that can be classed as installed
+		auto Manifest = BPSModule->LoadManifestFromFile(FoundManifests[0]);
+		if (!Manifest.IsValid())
+		{
+			//Something is wrong, suggests corruption so mark the folder for delete
+			ManifestsToRemove.Append(FoundManifests);
+			return true;
+		}
+		auto ChunkIDField = Manifest->GetCustomField("ChunkID");
+		if (!ChunkIDField.IsValid())
+		{
+			//Something is wrong, suggests corruption so mark the folder for delete
+			ManifestsToRemove.Append(FoundManifests);
+			return true;
+		}
+		auto ChunkPatchField = Manifest->GetCustomField("bIsPatch");
+		bool bIsPatch = ChunkPatchField.IsValid() ? ChunkPatchField->AsString() == TEXT("true") : false;
+		uint32 ChunkID = (uint32)ChunkIDField->AsInteger();
+
+		if (Pass == 1)
+		{
+			InstalledChunks.AddUnique(ChunkID, Manifest);
+		}
+		else if (Pass == 0 && ContentDir != InstallDir)
+		{
+			FString ChunkFdrName = FString::Printf(TEXT("%s%d"), !bIsPatch ? TEXT("base") : TEXT("patch"), ChunkID);
+			FString DestDir = *FPaths::Combine(*ContentDir, *ChunkFdrName);
+			if (PlatformFile.DirectoryExists(*DestDir))
+			{
+				PlatformFile.DeleteDirectoryRecursively(*DestDir);
+			}
+			PlatformFile.CreateDirectoryTree(*DestDir);
+			if (PlatformFile.CopyDirectoryTree(*DestDir, FilenameOrDirectory, true))
+			{
+				ManifestsToRemove.Add(FilenameOrDirectory);
+			}
+		}
+		else if (Pass == 2)
+		{
+			HoldingChunks.AddUnique(ChunkID, Manifest);
+		}
+
+		return true;
+	}
+
+	static const TCHAR *Name()
+	{
+		return TEXT("FChunkSetup");
+	}
+
+	FORCEINLINE TStatId GetStatId() const
+	{
+		RETURN_QUICK_DECLARE_CYCLE_STAT(FChunkSetupTask, STATGROUP_ThreadPoolAsyncTasks);
+	}
+
+	uint32 Run()
+	{
+		DoWork();
+		return 0;
+	}
+
+	bool IsDone()
+	{
+		return CompleteEvent->Wait(FTimespan(0));
+	}
+};
+
+class FChunkMountTask : public FRunnable, public IPlatformFile::FDirectoryVisitor
+{
+public:
+	/** Input parameters */
+	IBuildPatchServicesModule*	BPSModule;
+	FString						ContentDir; // Directory where installed chunks need to live to be mounted
+	const TArray<FString>*		CurrentMountPaks; // Array of already mounted Paks
+	const TSet<uint32>*			ExpectedChunks; //Manifests expected to be seen. Chunks not in this list are deleted
+	/** Output */
+	FEvent*									CompleteEvent;
+	TArray<FString>							MountedPaks;
+	/** Working */
+	TArray<FString>		FoundPaks;
+	TArray<FString>		FoundManifests;
+	TArray<FString>		ChunkInstallToDestroy;
+	FFileSearchVisitor	PakVisitor;
+	FFileSearchVisitor	ManifestVisitor;
+
+	FChunkMountTask()
+		: BPSModule(nullptr)
+		, CurrentMountPaks(nullptr)
+		, PakVisitor(TEXT("*.pak"), FoundPaks)
+		, ManifestVisitor(TEXT("*.manifest"), FoundManifests)
+	{
+		CompleteEvent = FPlatformProcess::GetSynchEventFromPool(true);
+	}
+
+	~FChunkMountTask()
+	{
+		FPlatformProcess::ReturnSynchEventToPool(CompleteEvent);
+	}
+
+	void SetupWork(IBuildPatchServicesModule* InBPSModule, FString InContentDir, const TArray<FString>& InCurrentMountedPaks, const TSet<uint32>& InExpectedChunks)
+	{
+		BPSModule = InBPSModule;
+		ContentDir = InContentDir;
+		CurrentMountPaks = &InCurrentMountedPaks;
+		ExpectedChunks = &InExpectedChunks;
+
+		MountedPaks.Reset();
+		FoundManifests.Reset();
+		FoundPaks.Reset();
+		ChunkInstallToDestroy.Reset();
+
+		CompleteEvent->Reset();
+	}
+
+	void DoWork()
+	{
+		IPlatformFile& PlatformFile = FPlatformFileManager::Get().GetPlatformFile();
+		PlatformFile.IterateDirectory(*ContentDir, *this);
+		for (const auto& Dir : ChunkInstallToDestroy)
+		{
+			PlatformFile.DeleteDirectoryRecursively(*Dir);
+		}
+		CompleteEvent->Trigger();
+	}
+
+	bool Visit(const TCHAR* FilenameOrDirectory, bool bIsDirectory) override
+	{
+		if (!bIsDirectory)
+		{
+			return true;
+		}
+		IPlatformFile& PlatformFile = FPlatformFileManager::Get().GetPlatformFile();
+
+		FoundManifests.Reset();
+		PlatformFile.IterateDirectory(FilenameOrDirectory, ManifestVisitor);
+		if (FoundManifests.Num() == 0)
+		{
+			return true;
+		}
+		// Only allow one manifest per folder, any more suggests corruption so ignore
+		if (FoundManifests.Num() > 1)
+		{
+			return true;
+		}
+		// Load the manifest, so that can be classed as installed
+		auto Manifest = BPSModule->LoadManifestFromFile(FoundManifests[0]);
+		if (!Manifest.IsValid())
+		{
+			//Something is wrong, suggests corruption so ignore
+			return true;
+		}
+		auto ChunkIDField = Manifest->GetCustomField("ChunkID");
+		if (!ChunkIDField.IsValid())
+		{
+			//Something is wrong, suggests corruption so ignore
+			return true;
+		}
+
+		if (!ExpectedChunks->Find(ChunkIDField->AsInteger())) 
+		{
+			//Add this to the list of chunks to remove
+			ChunkInstallToDestroy.Add(FilenameOrDirectory);
+			return true;
+		}
+
+		FoundPaks.Reset();
+		PlatformFile.IterateDirectoryRecursively(FilenameOrDirectory, PakVisitor);
+		if (FoundPaks.Num() == 0)
+		{
+			return true;
+		}
+		auto PakReadOrderField = Manifest->GetCustomField("PakReadOrdering");
+		uint32 PakReadOrder = PakReadOrderField.IsValid() ? (uint32)PakReadOrderField->AsInteger() : 0;
+		for (const auto& PakPath : FoundPaks)
+		{
+			//Note: A side effect here is that any unmounted paks get mounted. This is desirable as
+			// any previously installed chunks get mounted 
+			if (!CurrentMountPaks->Contains(PakPath) && !MountedPaks.Contains(PakPath))
+			{
+				if (FCoreDelegates::OnMountPak.IsBound())
+				{
+					auto bSuccess = FCoreDelegates::OnMountPak.Execute(PakPath, PakReadOrder, nullptr);
+#if !UE_BUILD_SHIPPING
+					if (!bSuccess)
+					{
+						// This can fail because of the sandbox system - which the pak system doesn't understand.
+						auto SandboxedPath = IFileManager::Get().ConvertToAbsolutePathForExternalAppForRead(*PakPath);
+						bSuccess = FCoreDelegates::OnMountPak.Execute(SandboxedPath, PakReadOrder, nullptr);
+					}
+#endif
+					//Register the install
+					BPSModule->RegisterAppInstallation(Manifest.ToSharedRef(), FilenameOrDirectory);
+				}
+			}
+		}
+		return true;
+	}
+
+	static const TCHAR *Name()
+	{
+		return TEXT("FChunkSetup");
+	}
+
+	FORCEINLINE TStatId GetStatId() const
+	{
+		RETURN_QUICK_DECLARE_CYCLE_STAT(FChunkMountTask, STATGROUP_ThreadPoolAsyncTasks);
+	}
+
+	uint32 Run()
+	{
+		DoWork();
+		return 0;
+	}
+
+	bool IsDone()
+	{
+		return CompleteEvent->Wait(FTimespan(0));
+	}
+};
\ No newline at end of file
diff --git a/Engine/Plugins/Runtime/HTTPChunkInstaller/Source/Public/CloudDelegateMacros.h b/Engine/Plugins/Runtime/HTTPChunkInstaller/Source/Public/CloudDelegateMacros.h
new file mode 100644
--- /dev/null
+++ b/Engine/Plugins/Runtime/HTTPChunkInstaller/Source/Public/CloudDelegateMacros.h
@@ -0,0 +1,76 @@
+// Copyright 1998-2017 Epic Games, Inc. All Rights Reserved.
+
+#pragma once
+
+#include "CoreMinimal.h"
+
+#define DEFINE_CLOUD_DELEGATE_BASE(DelegateName) \
+public: \
+	F##DelegateName DelegateName##Delegates; \
+public: \
+	virtual FDelegateHandle Add##DelegateName##Delegate_Handle(const F##DelegateName##Delegate& Delegate) \
+	{ \
+		DelegateName##Delegates.Add(Delegate); \
+		return Delegate.GetHandle(); \
+	} \
+	virtual void Clear##DelegateName##Delegate_Handle(FDelegateHandle& Handle) \
+	{ \
+		DelegateName##Delegates.Remove(Handle); \
+		Handle.Reset(); \
+	}
+
+#define DEFINE_CLOUD_DELEGATE(DelegateName) \
+	DEFINE_CLOUD_DELEGATE_BASE(DelegateName) \
+	virtual void Trigger##DelegateName##Delegates() \
+	{ \
+		DelegateName##Delegates.Broadcast(); \
+	}
+
+#define DEFINE_CLOUD_DELEGATE_ONE_PARAM(DelegateName, Param1Type) \
+	DEFINE_CLOUD_DELEGATE_BASE(DelegateName) \
+	virtual void Trigger##DelegateName##Delegates(Param1Type Param1) \
+	{ \
+		DelegateName##Delegates.Broadcast(Param1); \
+	}
+	
+#define DEFINE_CLOUD_DELEGATE_TWO_PARAM(DelegateName, Param1Type, Param2Type) \
+	DEFINE_CLOUD_DELEGATE_BASE(DelegateName) \
+	virtual void Trigger##DelegateName##Delegates(Param1Type Param1, Param2Type Param2) \
+	{ \
+		DelegateName##Delegates.Broadcast(Param1, Param2); \
+	}
+
+#define DEFINE_CLOUD_DELEGATE_THREE_PARAM(DelegateName, Param1Type, Param2Type, Param3Type) \
+	DEFINE_CLOUD_DELEGATE_BASE(DelegateName) \
+	virtual void Trigger##DelegateName##Delegates(Param1Type Param1, Param2Type Param2, Param3Type Param3) \
+	{ \
+		DelegateName##Delegates.Broadcast(Param1, Param2, Param3); \
+	}
+
+#define DEFINE_CLOUD_DELEGATE_FOUR_PARAM(DelegateName, Param1Type, Param2Type, Param3Type, Param4Type) \
+	DEFINE_CLOUD_DELEGATE_BASE(DelegateName) \
+	virtual void Trigger##DelegateName##Delegates(Param1Type Param1, Param2Type Param2, Param3Type Param3, Param4Type Param4) \
+	{ \
+		DelegateName##Delegates.Broadcast(Param1, Param2, Param3, Param4); \
+	}
+
+#define DEFINE_CLOUD_DELEGATE_FIVE_PARAM(DelegateName, Param1Type, Param2Type, Param3Type, Param4Type, Param5Type) \
+	DEFINE_CLOUD_DELEGATE_BASE(DelegateName) \
+	virtual void Trigger##DelegateName##Delegates(Param1Type Param1, Param2Type Param2, Param3Type Param3, Param4Type Param4, Param5Type Param5) \
+	{ \
+		DelegateName##Delegates.Broadcast(Param1, Param2, Param3, Param4, Param5); \
+	}
+
+#define DEFINE_CLOUD_DELEGATE_SIX_PARAM(DelegateName, Param1Type, Param2Type, Param3Type, Param4Type, Param5Type, Param6Type) \
+	DEFINE_CLOUD_DELEGATE_BASE(DelegateName) \
+	virtual void Trigger##DelegateName##Delegates(Param1Type Param1, Param2Type Param2, Param3Type Param3, Param4Type Param4, Param5Type Param5, Param6Type Param6) \
+	{ \
+		DelegateName##Delegates.Broadcast(Param1, Param2, Param3, Param4, Param5, Param6); \
+	}
+
+#define DEFINE_CLOUD_DELEGATE_SEVEN_PARAM(DelegateName, Param1Type, Param2Type, Param3Type, Param4Type, Param5Type, Param6Type, Param7Type) \
+	DEFINE_CLOUD_DELEGATE_BASE(DelegateName) \
+	virtual void Trigger##DelegateName##Delegates(Param1Type Param1, Param2Type Param2, Param3Type Param3, Param4Type Param4, Param5Type Param5, Param6Type Param6, Param7Type Param7) \
+	{ \
+		DelegateName##Delegates.Broadcast(Param1, Param2, Param3, Param4, Param5, Param6, Param7); \
+	}
diff --git a/Engine/Plugins/Runtime/HTTPChunkInstaller/Source/Public/CloudTitleFileInterface.h b/Engine/Plugins/Runtime/HTTPChunkInstaller/Source/Public/CloudTitleFileInterface.h
new file mode 100644
--- /dev/null
+++ b/Engine/Plugins/Runtime/HTTPChunkInstaller/Source/Public/CloudTitleFileInterface.h
@@ -0,0 +1,267 @@
+// Copyright 1998-2017 Epic Games, Inc. All Rights Reserved.
+
+#pragma once
+
+#include "CoreMinimal.h"
+#include "CloudDelegateMacros.h"
+
+/**
+* Paging info needed for a request that can return paged results
+*/
+class FCloudPagedQuery
+{
+public:
+	FCloudPagedQuery(int32 InStart = 0, int32 InCount = -1)
+		: Start(InStart)
+		, Count(InCount)
+	{}
+
+	/** @return true if valid range */
+	bool IsValidRange() const
+	{
+		return Start >= 0 && Count >= 0;
+	}
+
+	/** first entry to fetch */
+	int32 Start;
+	/** total entries to fetch. -1 means ALL */
+	int32 Count;
+};
+
+/** Holds metadata about a given downloadable file */
+struct FCloudHeader
+{
+	/** Hash value, if applicable, of the given file contents */
+	FString Hash;
+	/** The hash algorithm used to sign this file */
+	FName HashType;
+	/** Filename as downloaded */
+	FString DLName;
+	/** Logical filename, maps to the downloaded filename */
+	FString FileName;
+	/** File size */
+	int32 FileSize;
+	/** The full URL to download the file if it is stored in a CDN or separate host site */
+	FString URL;
+	/** The chunk id this file represents */
+	uint32 ChunkID;
+
+	/** Constructors */
+	FCloudHeader() :
+		FileSize(0),
+		ChunkID(0)
+	{}
+
+	FCloudHeader(const FString& InFileName, const FString& InDLName, int32 InFileSize) :
+		DLName(InDLName),
+		FileName(InFileName),
+		FileSize(InFileSize),
+		ChunkID(0)
+	{}
+
+	bool operator==(const FCloudHeader& Other) const
+	{
+		return FileSize == Other.FileSize &&
+			Hash == Other.Hash &&
+			HashType == Other.HashType &&
+			DLName == Other.DLName &&
+			FileName == Other.FileName &&
+			URL == Other.URL &&
+			ChunkID == Other.ChunkID;
+	}
+
+	bool operator<(const FCloudHeader& Other) const
+	{
+		return FileName.Compare(Other.FileName, ESearchCase::IgnoreCase) < 0;
+	}
+};
+
+/** The state of an async task (read friends, read content, write cloud file, etc) request */
+namespace ECloudAsyncTaskState
+{
+	enum Type
+	{
+		/** The task has not been started */
+		NotStarted,
+		/** The task is currently being processed */
+		InProgress,
+		/** The task has completed successfully */
+		Done,
+		/** The task failed to complete */
+		Failed
+	};
+
+	/** @return the stringified version of the enum passed in */
+	inline const TCHAR* ToString(ECloudAsyncTaskState::Type EnumVal)
+	{
+		switch (EnumVal)
+		{
+		case NotStarted:
+		{
+			return TEXT("NotStarted");
+		}
+		case InProgress:
+		{
+			return TEXT("InProgress");
+		}
+		case Done:
+		{
+			return TEXT("Done");
+		}
+		case Failed:
+		{
+			return TEXT("Failed");
+		}
+		}
+		return TEXT("");
+	}
+};
+
+/** Holds the data used in downloading a file asynchronously from the online service */
+struct FCloudEntry
+{
+	/** The name of the file as requested */
+	FString FileName;
+	/** The async state the file download is in */
+	ECloudAsyncTaskState::Type AsyncState;
+	/** The buffer of data for the file */
+	TArray<uint8> Data;
+
+	/** Constructors */
+	FCloudEntry() :
+		AsyncState(ECloudAsyncTaskState::NotStarted)
+	{
+	}
+
+	FCloudEntry(const FString& InFileName) :
+		FileName(InFileName),
+		AsyncState(ECloudAsyncTaskState::NotStarted)
+	{
+	}
+
+	virtual ~FCloudEntry() {}
+};
+
+/**
+ * Delegate fired when the list of files has been returned from the network store
+ *
+ * @param bWasSuccessful whether the file list was successful or not
+ *
+ */
+DECLARE_MULTICAST_DELEGATE_OneParam(FOnEnumerateFilesComplete, bool);
+typedef FOnEnumerateFilesComplete::FDelegate FOnEnumerateFilesCompleteDelegate;
+
+/**
+ * Delegate fired when as file read from the network platform's storage progresses
+ *
+ * @param FileName the name of the file this was for
+ * @param NumBytes the number of bytes read so far
+ */
+DECLARE_MULTICAST_DELEGATE_TwoParams(FOnReadFileProgress, const FString&, uint64);
+typedef FOnReadFileProgress::FDelegate FOnReadFileProgressDelegate;
+
+/**
+ * Delegate fired when a file read from the network platform's storage is complete
+ *
+ * @param bWasSuccessful whether the file read was successful or not
+ * @param FileName the name of the file this was for
+ *
+ */
+DECLARE_MULTICAST_DELEGATE_TwoParams(FOnReadFileComplete, bool, const FString&);
+typedef FOnReadFileComplete::FDelegate FOnReadFileCompleteDelegate;
+
+class ICloudTitleFile
+{
+protected:
+	ICloudTitleFile() {};
+
+public:
+	virtual ~ICloudTitleFile() {};
+
+	/**
+	 * Copies the file data into the specified buffer for the specified file
+	 *
+	 * @param FileName the name of the file to read
+	 * @param FileContents the out buffer to copy the data into
+	 *
+ 	 * @return true if the data was copied, false otherwise
+	 */
+	virtual bool GetFileContents(const FString& FileName, TArray<uint8>& FileContents) = 0;
+
+	/**
+	 * Empties the set of downloaded files if possible (no async tasks outstanding)
+	 *
+	 * @return true if they could be deleted, false if they could not
+	 */
+	virtual bool ClearFiles() = 0;
+
+	/**
+	 * Empties the cached data for this file if it is not being downloaded currently
+	 *
+	 * @param FileName the name of the file to remove from the cache
+	 *
+	 * @return true if it could be deleted, false if it could not
+	 */
+	virtual bool ClearFile(const FString& FileName) = 0;
+
+	/**
+	 * Delete cached files on disk
+	 *
+	 * @param bSkipEnumerated if true then only non-enumerated files are deleted
+	 */
+	virtual void DeleteCachedFiles(bool bSkipEnumerated) = 0;
+
+	/**
+	* Requests a list of available files from the network store
+	*
+	* @param Page paging info to use for query
+	*
+	* @return true if the request has started, false if not
+	*/
+	virtual bool EnumerateFiles(const FCloudPagedQuery& Page = FCloudPagedQuery()) = 0;
+
+	/**
+	 * Delegate fired when the list of files has been returned from the network store
+	 *
+	 * @param bWasSuccessful whether the file list was successful or not
+	 *
+	 */
+	DEFINE_CLOUD_DELEGATE_ONE_PARAM(OnEnumerateFilesComplete, bool);
+
+	/**
+	 * Returns the list of files that was returned by the network store
+	 * 
+	 * @param Files out array of file metadata
+	 *
+	 */
+	virtual void GetFileList(TArray<FCloudHeader>& Files) = 0;
+
+	/**
+	 * Starts an asynchronous read of the specified file from the network platform's file store
+	 *
+	 * @param FileToRead the name of the file to read
+	 *
+	 * @return true if the calls starts successfully, false otherwise
+	 */
+	virtual bool ReadFile(const FString& FileName) = 0;
+
+	/**
+	 * Delegate fired when a file read from the network platform's storage is complete
+	 *
+	 * @param bWasSuccessful whether the file read was successful or not
+	 * @param FileName the name of the file this was for
+	 *
+	 */
+	DEFINE_CLOUD_DELEGATE_TWO_PARAM(OnReadFileComplete, bool, const FString&);
+
+	/**
+	 * Delegate fired when as file read from the network platform's storage progresses
+	 *
+	 * @param FileName the name of the file this was for
+	 * @param NumBytes the number of bytes read so far
+	 */
+	DEFINE_CLOUD_DELEGATE_TWO_PARAM(OnReadFileProgress, const FString&, uint64);
+};
+
+typedef TSharedPtr<ICloudTitleFile, ESPMode::ThreadSafe> ICloudTitleFilePtr;
+typedef TSharedRef<ICloudTitleFile, ESPMode::ThreadSafe> ICloudTitleFileRef;
diff --git a/Engine/Plugins/Runtime/HTTPChunkInstaller/Source/Public/HTTPChunkInstaller.h b/Engine/Plugins/Runtime/HTTPChunkInstaller/Source/Public/HTTPChunkInstaller.h
new file mode 100644
--- /dev/null
+++ b/Engine/Plugins/Runtime/HTTPChunkInstaller/Source/Public/HTTPChunkInstaller.h
@@ -0,0 +1,167 @@
+// Copyright 1998-2016 Epic Games, Inc. All Rights Reserved.
+
+#pragma once
+
+#include "GenericPlatformChunkInstall.h"
+#include "UniquePtr.h"
+#include "CloudTitleFileInterface.h"
+#include "ChunkInstall.h"
+#include "ChunkSetup.h"
+#include "Ticker.h"
+
+/**
+* HTTP based implementation of chunk based install
+**/
+class HTTPCHUNKINSTALLER_API FHTTPChunkInstall : public IPlatformChunkInstall, public FTickerObjectBase
+{
+public:
+	FHTTPChunkInstall();
+	~FHTTPChunkInstall();
+
+	virtual EChunkLocation::Type GetChunkLocation( uint32 ChunkID ) override;
+
+	virtual bool GetProgressReportingTypeSupported(EChunkProgressReportingType::Type ReportType) override
+	{
+		if (ReportType == EChunkProgressReportingType::PercentageComplete)
+		{
+			return true;
+		}
+
+		return false;
+	}
+
+	virtual float GetChunkProgress( uint32 ChunkID, EChunkProgressReportingType::Type ReportType ) override;
+
+	virtual EChunkInstallSpeed::Type GetInstallSpeed() override
+	{
+		return InstallSpeed;
+	}
+
+	virtual bool SetInstallSpeed( EChunkInstallSpeed::Type InInstallSpeed ) override
+	{
+		if (InstallSpeed != InInstallSpeed)
+		{
+			InstallSpeed = InInstallSpeed;
+			if (InstallService.IsValid() && 
+				((InstallSpeed == EChunkInstallSpeed::Paused && !InstallService->IsPaused()) || (InstallSpeed != EChunkInstallSpeed::Paused && InstallService->IsPaused())))
+			{
+				InstallService->TogglePauseInstall();
+			}
+		}
+		return true;
+	}
+	
+	virtual bool PrioritizeChunk( uint32 ChunkID, EChunkPriority::Type Priority ) override;
+	virtual FDelegateHandle SetChunkInstallDelgate(uint32 ChunkID, FPlatformChunkInstallCompleteDelegate Delegate) override;
+	virtual void RemoveChunkInstallDelgate(uint32 ChunkID, FDelegateHandle Delegate) override;
+
+	virtual bool DebugStartNextChunk()
+	{
+		/** Unless paused we are always installing! */
+		InstallerState = ChunkInstallState::ReadTitleFiles;
+		return false;
+	}
+
+	bool Tick(float DeltaSeconds) override;
+
+	void EndInstall();
+
+private:
+
+	void InitialiseSystem();
+	bool AddDataToFileCache(const FString& ManifestHash,const TArray<uint8>& Data);
+	bool IsDataInFileCache(const FString& ManifestHash);
+	bool GetDataFromFileCache(const FString& ManifestHash,TArray<uint8>& Data);
+	bool RemoveDataFromFileCache(const FString& ManifestHash);
+	void UpdatePendingInstallQueue();
+	IBuildManifestPtr FindPreviousInstallManifest(const IBuildManifestPtr& ChunkManifest);
+	void BeginChunkInstall(uint32 NextChunkID,IBuildManifestPtr ChunkManifest, IBuildManifestPtr PrevInstallChunkManifest);
+	void ParseTitleFileManifest(const FString& ManifestFileHash);
+	bool BuildChunkFolderName(IBuildManifestRef Manifest, FString& ChunkFdrName, FString& ManifestName, uint32& ChunkID, bool& bIsPatch);
+	void OSSEnumerateFilesComplete(bool bSuccess);
+	void OSSReadFileComplete(bool bSuccess, const FString& Filename);
+	void OSSInstallComplete(bool, IBuildManifestRef);
+
+	DECLARE_MULTICAST_DELEGATE_OneParam(FPlatformChunkInstallCompleteMultiDelegate, uint32);
+
+	enum struct ChunkInstallState
+	{
+		Setup,
+		SetupWait,
+		QueryRemoteManifests,
+		EnterOfflineMode,
+		MoveInstalledChunks,
+		RequestingTitleFiles,
+		SearchTitleFiles,
+		ReadTitleFiles,
+		WaitingOnRead,
+		ReadComplete,
+		PostSetup,
+		Idle,
+		Installing,
+		CopyToContent,
+	};
+
+	struct FChunkPrio
+	{
+		FChunkPrio(uint32 InChunkID, EChunkPriority::Type InChunkPrio)
+			: ChunkID(InChunkID)
+			, ChunkPrio(InChunkPrio)
+		{
+
+		}
+
+		uint32					ChunkID;
+		EChunkPriority::Type	ChunkPrio;
+
+		bool operator == (const FChunkPrio& RHS) const
+		{
+			return ChunkID == RHS.ChunkID;
+		}
+
+		bool operator < (const FChunkPrio& RHS) const
+		{
+			return ChunkPrio < RHS.ChunkPrio;
+		}
+	};
+
+	FChunkInstallTask											ChunkCopyInstall;
+	TUniquePtr<FRunnableThread>									ChunkCopyInstallThread;
+	FChunkSetupTask												ChunkSetupTask;
+	TUniquePtr<FRunnableThread>									ChunkSetupTaskThread;
+	FChunkMountTask												ChunkMountTask;
+	TUniquePtr<FRunnableThread>									ChunkMountTaskThread;
+	TMultiMap<uint32, IBuildManifestPtr>						InstalledManifests;
+	TMultiMap<uint32, IBuildManifestPtr>						PrevInstallManifests;
+	TMultiMap<uint32, IBuildManifestPtr>						RemoteManifests;
+	TMap<uint32, FPlatformChunkInstallCompleteMultiDelegate>	DelegateMap;
+	TSet<FString>												ManifestsInMemory;
+	TSet<uint32>												ExpectedChunks;
+	TArray<FCloudHeader>										TitleFilesToRead;
+	TArray<uint8>												FileContentBuffer;
+	TArray<FChunkPrio>											PriorityQueue;
+	TArray<FString>												MountedPaks;
+	ICloudTitleFilePtr											OnlineTitleFile;
+	ICloudTitleFilePtr											OnlineTitleFileHttp;
+	IBuildInstallerPtr											InstallService;
+	IBuildManifestPtr											InstallingChunkManifest;
+	FDelegateHandle												EnumFilesCompleteHandle;
+	FDelegateHandle												ReadFileCompleteHandle;
+	FString														CloudDir;
+	FString														CloudDirectory;
+	FString														StageDir;
+	FString														InstallDir;
+	FString														BackupDir;
+	FString														ContentDir;
+	FString														CacheDir;
+	FString														HoldingDir;
+	IBuildPatchServicesModule*									BPSModule;
+	uint32														InstallingChunkID;
+	ChunkInstallState											InstallerState;
+	EChunkInstallSpeed::Type									InstallSpeed;
+	bool														bFirstRun;
+	bool														bSystemInitialised;
+#if !UE_BUILD_SHIPPING
+	bool														bDebugNoInstalledRequired;
+#endif
+};
\ No newline at end of file
diff --git a/Engine/Plugins/Runtime/HTTPChunkInstaller/Source/Public/HTTPChunkInstallerLog.h b/Engine/Plugins/Runtime/HTTPChunkInstaller/Source/Public/HTTPChunkInstallerLog.h
new file mode 100644
--- /dev/null
+++ b/Engine/Plugins/Runtime/HTTPChunkInstaller/Source/Public/HTTPChunkInstallerLog.h
@@ -0,0 +1,7 @@
+// Copyright 1998-2017 Epic Games, Inc. All Rights Reserved.
+
+#pragma once
+
+#include "LogMacros.h"
+
+DECLARE_LOG_CATEGORY_EXTERN(LogHTTPChunkInstaller, Log, All);
diff --git a/Engine/Plugins/Runtime/HTTPChunkInstaller/Source/Public/LocalTitleFile.h b/Engine/Plugins/Runtime/HTTPChunkInstaller/Source/Public/LocalTitleFile.h
new file mode 100644
--- /dev/null
+++ b/Engine/Plugins/Runtime/HTTPChunkInstaller/Source/Public/LocalTitleFile.h
@@ -0,0 +1,32 @@
+// Copyright 1998-2016 Epic Games, Inc. All Rights Reserved.
+
+#pragma once
+
+#include "CloudTitleFileInterface.h"
+
+class FLocalTitleFile : public ICloudTitleFile
+{
+public:
+
+	FLocalTitleFile(const FString& InRootDirectory);
+
+	// ICloudTitleFile interface
+
+	virtual bool GetFileContents(const FString& DLName, TArray<uint8>& FileContents) override;
+	virtual bool ClearFiles() override;
+	virtual bool ClearFile(const FString& DLName) override;
+	virtual void DeleteCachedFiles(bool bSkipEnumerated) override;
+	virtual bool EnumerateFiles(const FCloudPagedQuery& Page = FCloudPagedQuery()) override;
+	virtual void GetFileList(TArray<FCloudHeader>& InFileHeaders) override;
+	virtual bool ReadFile(const FString& DLName) override;
+
+protected:
+
+	FString GetFileNameFromDLName(const FString& DLName) const;
+
+private:
+
+	FString							RootDirectory;
+	TArray<FCloudHeader>			FileHeaders;
+	TMap< FString, TArray<uint8> >	DLNameToFileContents;
+};
\ No newline at end of file
