fix for Mac bild falure in Main

diff --git a/Engine/Plugins/Runtime/HTTPChunkInstaller/Source/Private/HTTPChunkInstaller.cpp b/Engine/Plugins/Runtime/HTTPChunkInstaller/Source/Private/HTTPChunkInstaller.cpp
--- a/Engine/Plugins/Runtime/HTTPChunkInstaller/Source/Private/HTTPChunkInstaller.cpp
+++ b/Engine/Plugins/Runtime/HTTPChunkInstaller/Source/Private/HTTPChunkInstaller.cpp
@@ -880,15 +880,6 @@ private:
 		return TEXT("http://") + BaseUrl + TEXT("/");
 	}
 
-	/** Info used to send request to authenticate a user *
-	struct FPendingEnumerate
-	{
-		FPendingEnumerate(const FPagedQuery& InPage = FPagedQuery())
-			: Page(InPage)
-		{
-		}
-		FPagedQuery Page;
-	};
 	/** List of pending Http requests for enumerating files */
 	TMap<FHttpRequestPtr, FCloudPagedQuery> EnumerateFilesRequests;
 
