fix for CIS warnings

diff --git a/Engine/Plugins/Runtime/HTTPChunkInstaller/Source/Private/HTTPChunkInstaller.cpp b/Engine/Plugins/Runtime/HTTPChunkInstaller/Source/Private/HTTPChunkInstaller.cpp
--- a/Engine/Plugins/Runtime/HTTPChunkInstaller/Source/Private/HTTPChunkInstaller.cpp
+++ b/Engine/Plugins/Runtime/HTTPChunkInstaller/Source/Private/HTTPChunkInstaller.cpp
@@ -1,9 +1,5 @@
 // Copyright 1998-2016 Epic Games, Inc. All Rights Reserved.
 
-#ifndef PLACEHOLDER_DLC_WORK
-#	define PLACEHOLDER_DLC_WORK 0
-#endif
-
 #include "HTTPChunkInstaller.h"
 #include "HTTPChunkInstallerLog.h"
 #include "ChunkInstall.h"
@@ -168,8 +164,8 @@ public:
 	* @param InSubsystem mcp subsystem being used
 	*/
 	FOnlineTitleFileHttp(const FString& InBaseUrl)
-		: BaseUrl(InBaseUrl)
-		, EnumerateFilesUrl(TEXT(""))
+		: EnumerateFilesUrl(TEXT(""))
+		, BaseUrl(InBaseUrl)
 	{
 		GConfig->GetString(TEXT("HTTPOnlineTitleFile"), TEXT("BaseUrl"), BaseUrl, GEngineIni);
 		GConfig->GetString(TEXT("HTTPOnlineTitleFile"), TEXT("EnumerateFilesUrl"), EnumerateFilesUrl, GEngineIni);
@@ -969,9 +965,11 @@ FHTTPChunkInstall::FHTTPChunkInstall()
 	: InstallingChunkID(-1)
 	, InstallerState(ChunkInstallState::Setup)
 	, InstallSpeed(EChunkInstallSpeed::Fast)
-	, bDebugNoInstalledRequired(false)
 	, bFirstRun(true)
 	, bSystemInitialised(false)
+#if !UE_BUILD_SHIPPING
+	, bDebugNoInstalledRequired(false)
+#endif
 {
 
 }
