Fix non-portable path errors in CIS.

diff --git a/Engine/Plugins/Runtime/HTTPChunkInstaller/Source/Private/HTTPChunkInstaller.cpp b/Engine/Plugins/Runtime/HTTPChunkInstaller/Source/Private/HTTPChunkInstaller.cpp
--- a/Engine/Plugins/Runtime/HTTPChunkInstaller/Source/Private/HTTPChunkInstaller.cpp
+++ b/Engine/Plugins/Runtime/HTTPChunkInstaller/Source/Private/HTTPChunkInstaller.cpp
@@ -8,7 +8,7 @@
 #include "SecureHash.h"
 #include "IHttpRequest.h"
 #include "IHttpResponse.h"
-#include "HTTPModule.h"
+#include "HttpModule.h"
 #include "JsonObject.h"
 #include "JsonReader.h"
 #include "JsonSerializer.h"
diff --git a/Engine/Plugins/Runtime/HTTPChunkInstaller/Source/Public/ChunkInstall.h b/Engine/Plugins/Runtime/HTTPChunkInstaller/Source/Public/ChunkInstall.h
--- a/Engine/Plugins/Runtime/HTTPChunkInstaller/Source/Public/ChunkInstall.h
+++ b/Engine/Plugins/Runtime/HTTPChunkInstaller/Source/Public/ChunkInstall.h
@@ -10,7 +10,7 @@
 #include "Event.h"
 #include "GenericPlatformFile.h"
 #include "Paths.h"
-#include "PlatformFileManager.h"
+#include "PlatformFilemanager.h"
 #include "CoreDelegates.h"
 #include "FileManager.h"
 #include "Stats.h"
