HTTPChunkInstaller plugin - disabled for HTML5

#rnx

diff --git a/Engine/Plugins/Runtime/HTTPChunkInstaller/HTTPChunkInstaller.uplugin b/Engine/Plugins/Runtime/HTTPChunkInstaller/HTTPChunkInstaller.uplugin
--- a/Engine/Plugins/Runtime/HTTPChunkInstaller/HTTPChunkInstaller.uplugin
+++ b/Engine/Plugins/Runtime/HTTPChunkInstaller/HTTPChunkInstaller.uplugin
@@ -14,7 +14,11 @@
         {
             "Name": "HTTPChunkInstaller",
             "Type": "Runtime",
-            "LoadingPhase" : "PostConfigInit"
+            "LoadingPhase" : "PostConfigInit",
+			"BlacklistPlatforms" :
+			[
+				"HTML5"
+			]
         }
     ],
 	"LocalizationTargets": [
@@ -23,4 +27,4 @@
 			"LoadingPolicy": "Always"
 		}
 	]
-}
\ No newline at end of file
+}
