Fix a cook warning from clashing localization keys.

diff --git a/Engine/Plugins/Runtime/HTTPChunkInstaller/Source/Private/HTTPChunkInstaller.cpp b/Engine/Plugins/Runtime/HTTPChunkInstaller/Source/Private/HTTPChunkInstaller.cpp
--- a/Engine/Plugins/Runtime/HTTPChunkInstaller/Source/Private/HTTPChunkInstaller.cpp
+++ b/Engine/Plugins/Runtime/HTTPChunkInstaller/Source/Private/HTTPChunkInstaller.cpp
@@ -678,7 +678,7 @@ private:
 		{
 			if (HttpResponse.IsValid())
 			{
-				ErrorStr = FText::Format(LOCTEXT("HttpResponse", "HTTP {0} response from {1}"),
+				ErrorStr = FText::Format(LOCTEXT("HttpResponseFromUrl", "HTTP {0} response from {1}"),
 					FText::AsNumber(HttpResponse->GetResponseCode()),
 					FText::FromString(HttpResponse->GetURL())).ToString();
 			}
@@ -742,7 +742,7 @@ private:
 		{
 			if (HttpResponse.IsValid())
 			{
-				ErrorStr = FText::Format(LOCTEXT("HttpResponse", "HTTP {0} response from {1}"),
+				ErrorStr = FText::Format(LOCTEXT("HttpResponseFromUrl", "HTTP {0} response from {1}"),
 					FText::AsNumber(HttpResponse->GetResponseCode()),
 					FText::FromString(HttpResponse->GetURL())).ToString();
 			}
