Optimized PakFile Index - fix deprecated OnMountPak in HTTPChunkInstaller

diff --git a/Engine/Plugins/Runtime/HTTPChunkInstaller/Source/Public/ChunkInstall.h b/Engine/Plugins/Runtime/HTTPChunkInstaller/Source/Public/ChunkInstall.h
--- a/Engine/Plugins/Runtime/HTTPChunkInstaller/Source/Public/ChunkInstall.h
+++ b/Engine/Plugins/Runtime/HTTPChunkInstaller/Source/Public/ChunkInstall.h
@@ -117,15 +117,15 @@ public:
 // 					// bump the read prioritiy 
 // 					++PakReadOrder;
 // 				}
-				if (FCoreDelegates::OnMountPak.IsBound())
+				if (FCoreDelegates::MountPak.IsBound())
 				{
-					auto bSuccess = FCoreDelegates::OnMountPak.Execute(PakFiles[PakIndex], PakReadOrder, nullptr);
+					auto bSuccess = FCoreDelegates::MountPak.Execute(PakFiles[PakIndex], PakReadOrder);
 #if !UE_BUILD_SHIPPING
 					if (!bSuccess)
 					{
 						// This can fail because of the sandbox system - which the pak system doesn't understand.
 						auto SandboxedPath = IFileManager::Get().ConvertToAbsolutePathForExternalAppForRead(*PakFiles[PakIndex]);
-						bSuccess = FCoreDelegates::OnMountPak.Execute(SandboxedPath, PakReadOrder, nullptr);
+						bSuccess = FCoreDelegates::MountPak.Execute(SandboxedPath, PakReadOrder);
 					}
 #endif
 					MountedPaks.Add(PakFiles[PakIndex]);
diff --git a/Engine/Plugins/Runtime/HTTPChunkInstaller/Source/Public/ChunkSetup.h b/Engine/Plugins/Runtime/HTTPChunkInstaller/Source/Public/ChunkSetup.h
--- a/Engine/Plugins/Runtime/HTTPChunkInstaller/Source/Public/ChunkSetup.h
+++ b/Engine/Plugins/Runtime/HTTPChunkInstaller/Source/Public/ChunkSetup.h
@@ -297,15 +297,15 @@ public:
 			// any previously installed chunks get mounted 
 			if (!CurrentMountPaks->Contains(PakPath) && !MountedPaks.Contains(PakPath))
 			{
-				if (FCoreDelegates::OnMountPak.IsBound())
+				if (FCoreDelegates::MountPak.IsBound())
 				{
-					auto bSuccess = FCoreDelegates::OnMountPak.Execute(PakPath, PakReadOrder, nullptr);
+					auto bSuccess = FCoreDelegates::MountPak.Execute(PakPath, PakReadOrder);
 #if !UE_BUILD_SHIPPING
 					if (!bSuccess)
 					{
 						// This can fail because of the sandbox system - which the pak system doesn't understand.
 						auto SandboxedPath = IFileManager::Get().ConvertToAbsolutePathForExternalAppForRead(*PakPath);
-						bSuccess = FCoreDelegates::OnMountPak.Execute(SandboxedPath, PakReadOrder, nullptr);
+						bSuccess = FCoreDelegates::MountPak.Execute(SandboxedPath, PakReadOrder);
 					}
 #endif
 					//Register the install
