remove deprecated per-chunk SetChunkInstallDelgate / RemoveChunkInstallDelgate functions.

diff --git a/Engine/Plugins/Runtime/HTTPChunkInstaller/Source/Private/HTTPChunkInstaller.cpp b/Engine/Plugins/Runtime/HTTPChunkInstaller/Source/Private/HTTPChunkInstaller.cpp
--- a/Engine/Plugins/Runtime/HTTPChunkInstaller/Source/Private/HTTPChunkInstaller.cpp
+++ b/Engine/Plugins/Runtime/HTTPChunkInstaller/Source/Private/HTTPChunkInstaller.cpp
@@ -1179,12 +1179,6 @@ bool FHTTPChunkInstall::Tick(float DeltaSeconds)
 			{
 				// No more manifests relating to the chunk ID are left to install.
 				// Inform any listeners that the install has been completed.
-				FPlatformChunkInstallCompleteMultiDelegate* FoundDelegate = DelegateMap.Find(InstallingChunkID);
-				if (FoundDelegate)
-				{
-					FoundDelegate->Broadcast(InstallingChunkID);
-				}
-
 				InstallDelegate.Broadcast(InstallingChunkID, true);
 			}
 			EndInstall();
@@ -1520,33 +1514,6 @@ bool FHTTPChunkInstall::PrioritizeChunk(uint32 ChunkID, EChunkPriority::Type Pri
 	return true;
 }
 
-FDelegateHandle FHTTPChunkInstall::SetChunkInstallDelgate(uint32 ChunkID, FPlatformChunkInstallCompleteDelegate Delegate)
-{
-	FPlatformChunkInstallCompleteMultiDelegate* FoundDelegate = DelegateMap.Find(ChunkID);
-	if (FoundDelegate)
-	{
-		return FoundDelegate->Add(Delegate);
-	}
-	else
-	{
-		FPlatformChunkInstallCompleteMultiDelegate MC;
-		auto RetVal = MC.Add(Delegate);
-		DelegateMap.Add(ChunkID, MC);
-		return RetVal;
-	}
-	return FDelegateHandle();
-}
-
-void FHTTPChunkInstall::RemoveChunkInstallDelgate(uint32 ChunkID, FDelegateHandle Delegate)
-{
-	FPlatformChunkInstallCompleteMultiDelegate* FoundDelegate = DelegateMap.Find(ChunkID);
-	if (!FoundDelegate)
-	{
-		return;
-	}
-	FoundDelegate->Remove(Delegate);
-}
-
 void FHTTPChunkInstall::BeginChunkInstall(uint32 ChunkID,IBuildManifestPtr ChunkManifest, IBuildManifestPtr PrevInstallChunkManifest)
 {
 	check(ChunkManifest->GetCustomField("ChunkID").IsValid());
diff --git a/Engine/Plugins/Runtime/HTTPChunkInstaller/Source/Public/HTTPChunkInstaller.h b/Engine/Plugins/Runtime/HTTPChunkInstaller/Source/Public/HTTPChunkInstaller.h
--- a/Engine/Plugins/Runtime/HTTPChunkInstaller/Source/Public/HTTPChunkInstaller.h
+++ b/Engine/Plugins/Runtime/HTTPChunkInstaller/Source/Public/HTTPChunkInstaller.h
@@ -50,8 +50,6 @@ public:
 	}
 	
 	virtual bool PrioritizeChunk( uint32 ChunkID, EChunkPriority::Type Priority ) override;
-	virtual FDelegateHandle SetChunkInstallDelgate(uint32 ChunkID, FPlatformChunkInstallCompleteDelegate Delegate) override;
-	virtual void RemoveChunkInstallDelgate(uint32 ChunkID, FDelegateHandle Delegate) override;
 
 	virtual bool DebugStartNextChunk()
 	{
@@ -82,8 +80,6 @@ private:
 
 	virtual EChunkLocation::Type GetChunkLocation(uint32 ChunkID) override;
 
-	DECLARE_MULTICAST_DELEGATE_OneParam(FPlatformChunkInstallCompleteMultiDelegate, uint32);
-
 	enum struct ChunkInstallState
 	{
 		Setup,
@@ -134,7 +130,6 @@ private:
 	TMultiMap<uint32, IBuildManifestPtr>						InstalledManifests;
 	TMultiMap<uint32, IBuildManifestPtr>						PrevInstallManifests;
 	TMultiMap<uint32, IBuildManifestPtr>						RemoteManifests;
-	TMap<uint32, FPlatformChunkInstallCompleteMultiDelegate>	DelegateMap;
 	TSet<FString>												ManifestsInMemory;
 	TSet<uint32>												ExpectedChunks;
 	TArray<FCloudHeader>										TitleFilesToRead;
diff --git a/Engine/Source/Runtime/Core/Public/GenericPlatform/GenericPlatformChunkInstall.h b/Engine/Source/Runtime/Core/Public/GenericPlatform/GenericPlatformChunkInstall.h
--- a/Engine/Source/Runtime/Core/Public/GenericPlatform/GenericPlatformChunkInstall.h
+++ b/Engine/Source/Runtime/Core/Public/GenericPlatform/GenericPlatformChunkInstall.h
@@ -205,12 +205,6 @@ public:
 	 */
 	virtual void RemoveChunkInstallDelegate( FDelegateHandle Delegate ) = 0;
 
-	UE_DEPRECATED(4.18, "Call AddChunkInstallDelegate instead, which is now bound for all chunk ids")
-	virtual FDelegateHandle SetChunkInstallDelgate( uint32 ChunkID, FPlatformChunkInstallCompleteDelegate Delegate ) = 0;
-
-	UE_DEPRECATED(4.18, "Call RemoveChunkInstallDelegate instead")
-	virtual void RemoveChunkInstallDelgate( uint32 ChunkID, FDelegateHandle Delegate ) = 0;
-
 
 	UE_DEPRECATED(5.2, "Call SupportsNamedChunkInstall instead")
 	virtual bool SupportsIntelligentInstall() = 0;
@@ -393,16 +387,6 @@ public:
 		InstallDelegate.Remove(Delegate);
 	}
 
-	virtual FDelegateHandle SetChunkInstallDelgate(uint32 ChunkID, FPlatformChunkInstallCompleteDelegate Delegate) override
-	{
-		return FDelegateHandle();
-	}
-
-	virtual void RemoveChunkInstallDelgate(uint32 ChunkID, FDelegateHandle Delegate) override
-	{
-		return;
-	}
-
 	virtual bool SupportsIntelligentInstall() override
 	{
 		return false;
