Update more code using FJsonObject to use TCHAR strings instead of ANSI strings. Removes unnecessary string conversions and removes compilation errors when FJsonObject is updated to use FStringViews.

diff --git a/Engine/Plugins/Developer/OneSkyLocalizationService/Source/OneSkyLocalizationService/Private/OneSkyLocalizationServiceSettings.cpp b/Engine/Plugins/Developer/OneSkyLocalizationService/Source/OneSkyLocalizationService/Private/OneSkyLocalizationServiceSettings.cpp
--- a/Engine/Plugins/Developer/OneSkyLocalizationService/Source/OneSkyLocalizationService/Private/OneSkyLocalizationServiceSettings.cpp
+++ b/Engine/Plugins/Developer/OneSkyLocalizationService/Source/OneSkyLocalizationService/Private/OneSkyLocalizationServiceSettings.cpp
@@ -91,20 +91,20 @@ void FOneSkyLocalizationServiceSettings::LoadSettings()
 				continue;
 			}
 
-			TArray<TSharedPtr<FJsonValue> > CredentialsArray = JSONObject->GetArrayField("Credentials");
+			TArray<TSharedPtr<FJsonValue> > CredentialsArray = JSONObject->GetArrayField(TEXT("Credentials"));
 
 			for (TSharedPtr<FJsonValue> CredentialsItem : CredentialsArray)
 			{
 				bool bFoundCredentials = false;
 				const TSharedPtr<FJsonObject>& CredentialsObject = CredentialsItem->AsObject();
-				if (CredentialsObject->HasTypedField<EJson::String>("Name"))
+				if (CredentialsObject->HasTypedField<EJson::String>(TEXT("Name")))
 				{
-					ConnectionInfo.Name = CredentialsObject->GetStringField("Name");
+					ConnectionInfo.Name = CredentialsObject->GetStringField(TEXT("Name"));
 				}
 
-				if (CredentialsObject->HasTypedField<EJson::String>("ApiKey"))
+				if (CredentialsObject->HasTypedField<EJson::String>(TEXT("ApiKey")))
 				{
-					ConnectionInfo.ApiKey = CredentialsObject->GetStringField("ApiKey");
+					ConnectionInfo.ApiKey = CredentialsObject->GetStringField(TEXT("ApiKey"));
 					bFoundCredentials = true;
 				}
 				else
@@ -112,9 +112,9 @@ void FOneSkyLocalizationServiceSettings::LoadSettings()
 					UE_LOG(LogLocalizationService, Log, TEXT("Credentials file %s is missing ApiKey for connection name %s"), *CredentialsFile, *ConnectionInfo.Name);
 				}
 
-				if (CredentialsObject->HasTypedField<EJson::String>("ApiSecret"))
+				if (CredentialsObject->HasTypedField<EJson::String>(TEXT("ApiSecret")))
 				{
-					ConnectionInfo.ApiSecret = CredentialsObject->GetStringField("ApiSecret");
+					ConnectionInfo.ApiSecret = CredentialsObject->GetStringField(TEXT("ApiSecret"));
 				}
 				else
 				{
diff --git a/Engine/Plugins/Experimental/Web/WebAPI/Source/WebAPI/Private/Security/WebAPIAuthentication.cpp b/Engine/Plugins/Experimental/Web/WebAPI/Source/WebAPI/Private/Security/WebAPIAuthentication.cpp
--- a/Engine/Plugins/Experimental/Web/WebAPI/Source/WebAPI/Private/Security/WebAPIAuthentication.cpp
+++ b/Engine/Plugins/Experimental/Web/WebAPI/Source/WebAPI/Private/Security/WebAPIAuthentication.cpp
@@ -115,11 +115,11 @@ bool FWebAPIOAuthSchemeHandler::HandleHttpResponse(EHttpResponseCodes::Type InRe
 				TSharedPtr<FJsonObject> JsonObject;
 				if(FJsonSerializer::Deserialize(JsonReader, JsonObject))
 				{
-					const int32 UnixTimeExpire = JsonObject->GetNumberField("expires_on");
+					const int32 UnixTimeExpire = JsonObject->GetNumberField(TEXT("expires_on"));
 
 					UE_LOG(LogWebAPI, Display, TEXT("Generate token Response success"));
-					AuthSettings->TokenType = JsonObject->GetStringField("token_type");
-					AuthSettings->AccessToken = JsonObject->GetStringField("access_token");
+					AuthSettings->TokenType = JsonObject->GetStringField(TEXT("token_type"));
+					AuthSettings->AccessToken = JsonObject->GetStringField(TEXT("access_token"));
 					AuthSettings->ExpiresOn =  FDateTime::FromUnixTimestamp(UnixTimeExpire);
 					AuthSettings->SaveConfig();
 					bSuccess = true;
diff --git a/Engine/Plugins/Media/ElectraCDM/Source/ElectraCDM/Private/ClearKey/ClearKeyCDM.cpp b/Engine/Plugins/Media/ElectraCDM/Source/ElectraCDM/Private/ClearKey/ClearKeyCDM.cpp
--- a/Engine/Plugins/Media/ElectraCDM/Source/ElectraCDM/Private/ClearKey/ClearKeyCDM.cpp
+++ b/Engine/Plugins/Media/ElectraCDM/Source/ElectraCDM/Private/ClearKey/ClearKeyCDM.cpp
@@ -786,15 +786,15 @@ ECDMError FClearKeyDRMClient::SetLicenseKeyResponseData(void* InEventId, int32 H
 			if (FJsonSerializer::Deserialize(Reader, KeyResponse))
 			{
 				TArray<TSharedPtr<FJsonValue>> Keys;
-				Keys = KeyResponse->GetArrayField("keys");
+				Keys = KeyResponse->GetArrayField(TEXT("keys"));
 				for(int32 i=0; i<Keys.Num(); ++i)
 				{
 					auto Key = Keys[i]->AsObject();
-					FString KTY = Key->GetStringField("kty");	// Mandatory (see RFC-7518 Section 6.1), must be "oct"
+					FString KTY = Key->GetStringField(TEXT("kty"));	// Mandatory (see RFC-7518 Section 6.1), must be "oct"
 					if (KTY.Equals(TEXT("oct")))
 					{
-						FString KID = Key->GetStringField("kid");
-						FString K = Key->GetStringField("k");
+						FString KID = Key->GetStringField(TEXT("kid"));
+						FString K = Key->GetStringField(TEXT("k"));
 						if (KID.Len() && K.Len())
 						{
 							FClearKeyKIDKey NewKey;
diff --git a/Engine/Plugins/Runtime/HTTPChunkInstaller/Source/Private/HTTPChunkInstaller.cpp b/Engine/Plugins/Runtime/HTTPChunkInstaller/Source/Private/HTTPChunkInstaller.cpp
--- a/Engine/Plugins/Runtime/HTTPChunkInstaller/Source/Private/HTTPChunkInstaller.cpp
+++ b/Engine/Plugins/Runtime/HTTPChunkInstaller/Source/Private/HTTPChunkInstaller.cpp
@@ -1212,7 +1212,7 @@ void FHTTPChunkInstall::UpdatePendingInstallQueue()
 		if (FoundChunkManifests.Num() > 0)
 		{
 			auto ChunkManifest = FoundChunkManifests[0];
-			auto ChunkIDField = ChunkManifest->GetCustomField("ChunkID");
+			auto ChunkIDField = ChunkManifest->GetCustomField(TEXT("ChunkID"));
 			if (ChunkIDField.IsValid())
 			{
 				BeginChunkInstall(NextChunk.ChunkID, ChunkManifest, FindPreviousInstallManifest(ChunkManifest));
@@ -1235,7 +1235,7 @@ void FHTTPChunkInstall::UpdatePendingInstallQueue()
 			if (It)
 			{
 				IBuildManifestPtr ChunkManifest = It.Value();
-				auto ChunkIDField = ChunkManifest->GetCustomField("ChunkID");
+				auto ChunkIDField = ChunkManifest->GetCustomField(TEXT("ChunkID"));
 				if (ChunkIDField.IsValid())
 				{
 					BeginChunkInstall(ChunkIDField->AsInteger(), ChunkManifest, FindPreviousInstallManifest(ChunkManifest));
@@ -1369,7 +1369,7 @@ void FHTTPChunkInstall::OSSInstallComplete(const IBuildInstallerRef& Installer)
 		InstalledManifests.MultiFind(ChunkID, FoundManifests);
 		for (const auto& It : FoundManifests)
 		{
-			auto FoundPatchField = It->GetCustomField("bIsPatch");
+			auto FoundPatchField = It->GetCustomField(TEXT("bIsPatch"));
 			bool bFoundPatch = FoundPatchField.IsValid() ? FoundPatchField->AsString() == TEXT("true") : false;
 			if (bFoundPatch == bIsPatch)
 			{
@@ -1403,7 +1403,7 @@ void FHTTPChunkInstall::ParseTitleFileManifest(const FString& ManifestFileHash)
 		UE_LOG(LogHTTPChunkInstaller, Warning, TEXT("Manifest was invalid"));
 		return;
 	}
-	auto RemoteChunkIDField = RemoteManifest->GetCustomField("ChunkID");
+	auto RemoteChunkIDField = RemoteManifest->GetCustomField(TEXT("ChunkID"));
 	if (!RemoteChunkIDField.IsValid())
 	{
 		UE_LOG(LogHTTPChunkInstaller, Warning, TEXT("Manifest ChunkID was invalid or missing"));
@@ -1417,14 +1417,14 @@ void FHTTPChunkInstall::ParseTitleFileManifest(const FString& ManifestFileHash)
 	uint32 FoundCount = FoundManifests.Num();
 	if (FoundCount > 0)
 	{
-		auto RemotePatchManifest = RemoteManifest->GetCustomField("bIsPatch");
+		auto RemotePatchManifest = RemoteManifest->GetCustomField(TEXT("bIsPatch"));
 		auto RemoteVersion = RemoteManifest->GetVersionString();
 		bool bRemoteIsPatch = RemotePatchManifest.IsValid() ? RemotePatchManifest->AsString() == TEXT("true") : false;
 		for (uint32 FoundIndex = 0; FoundIndex < FoundCount; ++FoundIndex)
 		{
 			const auto& InstalledManifest = FoundManifests[FoundIndex];
 			auto InstalledVersion = InstalledManifest->GetVersionString();
-			auto InstallPatchManifest = InstalledManifest->GetCustomField("bIsPatch");
+			auto InstallPatchManifest = InstalledManifest->GetCustomField(TEXT("bIsPatch"));
 			bool bInstallIsPatch = InstallPatchManifest.IsValid() ? InstallPatchManifest->AsString() == TEXT("true") : false;
 			if (InstalledVersion != RemoteVersion && bInstallIsPatch == bRemoteIsPatch)
 			{
@@ -1470,8 +1470,8 @@ void FHTTPChunkInstall::ParseTitleFileManifest(const FString& ManifestFileHash)
 
 bool FHTTPChunkInstall::BuildChunkFolderName(IBuildManifestRef Manifest, FString& ChunkFdrName, FString& ManifestName, uint32& ChunkID, bool& bIsPatch)
 {
-	auto ChunkIDField = Manifest->GetCustomField("ChunkID");
-	auto ChunkPatchField = Manifest->GetCustomField("bIsPatch");
+	auto ChunkIDField = Manifest->GetCustomField(TEXT("ChunkID"));
+	auto ChunkPatchField = Manifest->GetCustomField(TEXT("bIsPatch"));
 
 	if (!ChunkIDField.IsValid())
 	{
@@ -1508,12 +1508,12 @@ bool FHTTPChunkInstall::PrioritizeChunk(uint32 ChunkID, EChunkPriority::Type Pri
 
 void FHTTPChunkInstall::BeginChunkInstall(uint32 ChunkID,IBuildManifestPtr ChunkManifest, IBuildManifestPtr PrevInstallChunkManifest)
 {
-	check(ChunkManifest->GetCustomField("ChunkID").IsValid());
+	check(ChunkManifest->GetCustomField(TEXT("ChunkID")).IsValid());
 	InstallingChunkID = ChunkID;
 	check(ChunkID > 0);
 	InstallingChunkManifest = ChunkManifest;
 	IPlatformFile& PlatformFile = FPlatformFileManager::Get().GetPlatformFile();
-	auto PatchField = ChunkManifest->GetCustomField("bIsPatch");
+	auto PatchField = ChunkManifest->GetCustomField(TEXT("bIsPatch"));
 	bool bIsPatch = PatchField.IsValid() ? PatchField->AsString() == TEXT("true") : false;
 	auto ChunkFolderName = FString::Printf(TEXT("%s%d"),!bIsPatch ? TEXT("base") : TEXT("patch"), InstallingChunkID);
 	auto ChunkInstallDir = FPaths::Combine(*InstallDir,*ChunkFolderName);
@@ -1672,7 +1672,7 @@ void FHTTPChunkInstall::InitialiseSystem()
 
 IBuildManifestPtr FHTTPChunkInstall::FindPreviousInstallManifest(const IBuildManifestPtr& ChunkManifest)
 {
-	auto ChunkIDField = ChunkManifest->GetCustomField("ChunkID");
+	auto ChunkIDField = ChunkManifest->GetCustomField(TEXT("ChunkID"));
 	if (!ChunkIDField.IsValid())
 	{
 		return IBuildManifestPtr();
diff --git a/Engine/Plugins/Runtime/HTTPChunkInstaller/Source/Public/ChunkInstall.h b/Engine/Plugins/Runtime/HTTPChunkInstaller/Source/Public/ChunkInstall.h
--- a/Engine/Plugins/Runtime/HTTPChunkInstaller/Source/Public/ChunkInstall.h
+++ b/Engine/Plugins/Runtime/HTTPChunkInstaller/Source/Public/ChunkInstall.h
@@ -105,7 +105,7 @@ public:
 		// Find all pak files.
 		FPakSearchVisitor Visitor(PakFiles);
 		PlatformFile.IterateDirectoryRecursively(*DestDir, Visitor);
-		auto PakReadOrderField = BuildManifest->GetCustomField("PakReadOrdering");
+		auto PakReadOrderField = BuildManifest->GetCustomField(TEXT("PakReadOrdering"));
 		uint32 PakReadOrder = PakReadOrderField.IsValid() ? (uint32)PakReadOrderField->AsInteger() : 0;
 		for (uint32 PakIndex = 0, PakCount = PakFiles.Num(); PakIndex < PakCount; ++PakIndex)
 		{
diff --git a/Engine/Plugins/Runtime/HTTPChunkInstaller/Source/Public/ChunkSetup.h b/Engine/Plugins/Runtime/HTTPChunkInstaller/Source/Public/ChunkSetup.h
--- a/Engine/Plugins/Runtime/HTTPChunkInstaller/Source/Public/ChunkSetup.h
+++ b/Engine/Plugins/Runtime/HTTPChunkInstaller/Source/Public/ChunkSetup.h
@@ -137,14 +137,14 @@ public:
 			ManifestsToRemove.Append(FoundManifests);
 			return true;
 		}
-		auto ChunkIDField = Manifest->GetCustomField("ChunkID");
+		auto ChunkIDField = Manifest->GetCustomField(TEXT("ChunkID"));
 		if (!ChunkIDField.IsValid())
 		{
 			//Something is wrong, suggests corruption so mark the folder for delete
 			ManifestsToRemove.Append(FoundManifests);
 			return true;
 		}
-		auto ChunkPatchField = Manifest->GetCustomField("bIsPatch");
+		auto ChunkPatchField = Manifest->GetCustomField(TEXT("bIsPatch"));
 		bool bIsPatch = ChunkPatchField.IsValid() ? ChunkPatchField->AsString() == TEXT("true") : false;
 		uint32 ChunkID = (uint32)ChunkIDField->AsInteger();
 
@@ -280,7 +280,7 @@ public:
 			//Something is wrong, suggests corruption so ignore
 			return true;
 		}
-		auto ChunkIDField = Manifest->GetCustomField("ChunkID");
+		auto ChunkIDField = Manifest->GetCustomField(TEXT("ChunkID"));
 		if (!ChunkIDField.IsValid())
 		{
 			//Something is wrong, suggests corruption so ignore
@@ -300,7 +300,7 @@ public:
 		{
 			return true;
 		}
-		auto PakReadOrderField = Manifest->GetCustomField("PakReadOrdering");
+		auto PakReadOrderField = Manifest->GetCustomField(TEXT("PakReadOrdering"));
 		uint32 PakReadOrder = PakReadOrderField.IsValid() ? (uint32)PakReadOrderField->AsInteger() : 0;
 		for (const auto& PakPath : FoundPaks)
 		{
diff --git a/Engine/Plugins/Tests/WidgetAutomationTests/Source/WidgetAutomationTests/Private/STextBlockTest.cpp b/Engine/Plugins/Tests/WidgetAutomationTests/Source/WidgetAutomationTests/Private/STextBlockTest.cpp
--- a/Engine/Plugins/Tests/WidgetAutomationTests/Source/WidgetAutomationTests/Private/STextBlockTest.cpp
+++ b/Engine/Plugins/Tests/WidgetAutomationTests/Source/WidgetAutomationTests/Private/STextBlockTest.cpp
@@ -77,7 +77,7 @@ public:
 					// Confirm that the payload tint calculated in OnPaint matches the color we set in the test.
 					if (TSharedPtr<FJsonObject> JsonObject = GetSubjectWidget()->GetTestJSONObject())
 					{
-						FString PayloadTint = JsonObject->GetObjectField("ShapedTextPayload")->GetStringField("Tint");
+						FString PayloadTint = JsonObject->GetObjectField(TEXT("ShapedTextPayload"))->GetStringField(TEXT("Tint"));
 						TestObj->AddErrorIfFalse(PayloadTint == FLinearColor(0.5, 0.4, 0.3, 0.9).ToString(), FString::Printf(TEXT("In %s: The tint of the payload does not match the passed color. Expected: %s, Actual: %s"), GetGeneratedTypeName<STextBlock>(), *FLinearColor(0.5, 0.4, 0.3, 0.9).ToString(), *PayloadTint));
 					}
 				}
diff --git a/Engine/Plugins/Tests/WidgetAutomationTests/Source/WidgetAutomationTests/Private/WidgetMockNonTemplate.cpp b/Engine/Plugins/Tests/WidgetAutomationTests/Source/WidgetAutomationTests/Private/WidgetMockNonTemplate.cpp
--- a/Engine/Plugins/Tests/WidgetAutomationTests/Source/WidgetAutomationTests/Private/WidgetMockNonTemplate.cpp
+++ b/Engine/Plugins/Tests/WidgetAutomationTests/Source/WidgetAutomationTests/Private/WidgetMockNonTemplate.cpp
@@ -93,10 +93,10 @@ void FWidgetMockNonTemplate::OnPaintDataValidation(const FGeometry& AllottedGeom
 	if (TestJsonObject.IsValid())
 	{
 		// Validate the parameters AllottedGeometry and InWidgetStyle.
-		if (TestJsonObject->GetStringField("AllottedGeometry") == AllottedGeometry.ToString()
-			&& TestJsonObject->GetObjectField("WidgetStyle")->GetStringField("ColorAndOpacityTint") == InWidgetStyle.GetColorAndOpacityTint().ToString()
-			&& TestJsonObject->GetObjectField("WidgetStyle")->GetStringField("ForegroundColor") == InWidgetStyle.GetForegroundColor().ToString()
-			&& TestJsonObject->GetObjectField("WidgetStyle")->GetStringField("SubduedForeground") == InWidgetStyle.GetSubduedForegroundColor().ToString())
+		if (TestJsonObject->GetStringField(TEXT("AllottedGeometry")) == AllottedGeometry.ToString()
+			&& TestJsonObject->GetObjectField(TEXT("WidgetStyle"))->GetStringField(TEXT("ColorAndOpacityTint")) == InWidgetStyle.GetColorAndOpacityTint().ToString()
+			&& TestJsonObject->GetObjectField(TEXT("WidgetStyle"))->GetStringField(TEXT("ForegroundColor")) == InWidgetStyle.GetForegroundColor().ToString()
+			&& TestJsonObject->GetObjectField(TEXT("WidgetStyle"))->GetStringField(TEXT("SubduedForeground")) == InWidgetStyle.GetSubduedForegroundColor().ToString())
 		{
 			// Validate the payloads in OutDrawElements.
 			for (FSlateCachedElementList* CachedElementList : OutDrawElements.GetCurrentCachedElementWithNewData())
@@ -106,8 +106,8 @@ void FWidgetMockNonTemplate::OnPaintDataValidation(const FGeometry& AllottedGeom
 					//@TODO Support all element types and create their own JSON conversion function.
 					for (const FSlateShapedTextElement& DrawElements : CachedElementList->DrawElements.Get<(uint8)EElementType::ET_ShapedText>())
 					{
-						if (TestJsonObject->GetObjectField("ShapedTextElement")->GetStringField("Tint") == DrawElements.GetTint().ToString()
-							&& TestJsonObject->GetObjectField("ShapedTextElement")->GetStringField("OutlineTint") == DrawElements.GetOutlineTint().ToString())
+						if (TestJsonObject->GetObjectField(TEXT("ShapedTextElement"))->GetStringField(TEXT("Tint")) == DrawElements.GetTint().ToString()
+							&& TestJsonObject->GetObjectField(TEXT("ShapedTextElement"))->GetStringField(TEXT("OutlineTint")) == DrawElements.GetOutlineTint().ToString())
 						{
 							if (!DataValidations.Contains(FTestFunctionNames::NAME_OnPaint))
 							{
@@ -156,7 +156,7 @@ void FWidgetMockNonTemplate::ComputeDesiredSizeDataValidation(float LayoutScaleM
 	if (TestJsonObject.IsValid())
 	{
 		// Validate the computed size in JSON against the size we just computed.
-		if (TestJsonObject->GetNumberField("ComputeDesiredSizeX") == RealComputedSize.X && TestJsonObject->GetNumberField("ComputeDesiredSizeY") == RealComputedSize.Y)
+		if (TestJsonObject->GetNumberField(TEXT("ComputeDesiredSizeX")) == RealComputedSize.X && TestJsonObject->GetNumberField(TEXT("ComputeDesiredSizeY")) == RealComputedSize.Y)
 		{
 			DataValidations.Add(FTestFunctionNames::NAME_ComputeDesiredSize, true);
 
diff --git a/Engine/Plugins/VirtualProduction/Switchboard/Source/SwitchboardEditor/Private/SwitchboardScriptInterop.cpp b/Engine/Plugins/VirtualProduction/Switchboard/Source/SwitchboardEditor/Private/SwitchboardScriptInterop.cpp
--- a/Engine/Plugins/VirtualProduction/Switchboard/Source/SwitchboardEditor/Private/SwitchboardScriptInterop.cpp
+++ b/Engine/Plugins/VirtualProduction/Switchboard/Source/SwitchboardEditor/Private/SwitchboardScriptInterop.cpp
@@ -208,14 +208,14 @@ void FSwitchboardVerifyResult::SetFromJson(FStringView JsonSV)
 		return;
 	}
 
-	if (!JsonObject->TryGetStringField("log", this->Log))
+	if (!JsonObject->TryGetStringField(TEXT("log"), this->Log))
 	{
 		UE_LOG(LogSwitchboardPlugin, Error, TEXT("FSwitchboardVerifyResult: 'log' field missing"));
 		this->Log.Empty();
 	}
 
 	bool bSuccess = false;
-	if (!JsonObject->TryGetBoolField("success", bSuccess))
+	if (!JsonObject->TryGetBoolField(TEXT("success"), bSuccess))
 	{
 		UE_LOG(LogSwitchboardPlugin, Error, TEXT("FSwitchboardVerifyResult: 'success' field missing"));
 		this->Summary = ESummary::InvalidOutput;
diff --git a/Engine/Source/Editor/UnrealEd/Private/Fbx/FbxSceneImportData.cpp b/Engine/Source/Editor/UnrealEd/Private/Fbx/FbxSceneImportData.cpp
--- a/Engine/Source/Editor/UnrealEd/Private/Fbx/FbxSceneImportData.cpp
+++ b/Engine/Source/Editor/UnrealEd/Private/Fbx/FbxSceneImportData.cpp
@@ -116,7 +116,7 @@ UnFbx::FBXImportOptions *JSONToFbxOption(TSharedPtr<FJsonValue> OptionJsonValue,
 	//Skeletal mesh options
 	OptionObj->TryGetBoolField(TEXT("bUpdateSkeletonReferencePose"), Option->bUpdateSkeletonReferencePose);
 	//TODO support T0AsRefPose
-	//OptionObj->TryGetBoolField("bUseT0AsRefPose", Option->bUseT0AsRefPose);
+	//OptionObj->TryGetBoolField(TEXT("bUseT0AsRefPose"), Option->bUseT0AsRefPose);
 	Option->bUseT0AsRefPose = false;
 	OptionObj->TryGetBoolField(TEXT("bPreserveSmoothingGroups"), Option->bPreserveSmoothingGroups);
 	OptionObj->TryGetBoolField(TEXT("bKeepSectionsSeparate"), Option->bKeepSectionsSeparate);
